Rename HierarchyTableBuiltin plugin to HierarchyTableAnimation

diff --git a/Engine/Plugins/Experimental/AnimNext/AnimNext.uplugin b/Engine/Plugins/Experimental/AnimNext/AnimNext.uplugin
--- a/Engine/Plugins/Experimental/AnimNext/AnimNext.uplugin
+++ b/Engine/Plugins/Experimental/AnimNext/AnimNext.uplugin
@@ -43,9 +43,9 @@
 			"Enabled": true
 		},
 		{
-			"Name": "HierarchyTableBuiltin",
+			"Name": "HierarchyTableAnimation",
 			"Enabled": true
-		},
+		}
 	],
 	"Modules" :
 	[
diff --git a/Engine/Plugins/Experimental/AnimNextAnimGraph/Source/AnimNextAnimGraph/AnimNextAnimGraph.Build.cs b/Engine/Plugins/Experimental/AnimNextAnimGraph/Source/AnimNextAnimGraph/AnimNextAnimGraph.Build.cs
--- a/Engine/Plugins/Experimental/AnimNextAnimGraph/Source/AnimNextAnimGraph/AnimNextAnimGraph.Build.cs
+++ b/Engine/Plugins/Experimental/AnimNextAnimGraph/Source/AnimNextAnimGraph/AnimNextAnimGraph.Build.cs
@@ -22,7 +22,7 @@ namespace UnrealBuildTool.Rules
 					"Engine",
 					"AnimNext",
 					"HierarchyTableRuntime",
-					"HierarchyTableBuiltinRuntime"
+					"HierarchyTableAnimationRuntime"
 				}
 			);
 
diff --git a/Engine/Plugins/Experimental/Animation/HierarchyTableAnimation/Config/BaseHierarchyTableAnimation.ini b/Engine/Plugins/Experimental/Animation/HierarchyTableAnimation/Config/BaseHierarchyTableAnimation.ini
new file mode 100644
--- /dev/null
+++ b/Engine/Plugins/Experimental/Animation/HierarchyTableAnimation/Config/BaseHierarchyTableAnimation.ini
@@ -0,0 +1,7 @@
+[CoreRedirects]
+
+; Rename HierarchyTableBuiltin to HierarchyTableAnimation
++StructRedirects=(OldName="/Script/HierarchyTableBuiltinRuntime.HierarchyTable_ElementType_Mask", NewName="/Script/HierarchyTableAnimationRuntime.HierarchyTable_ElementType_Mask")
++StructRedirects=(OldName="/Script/HierarchyTableBuiltinRuntime.HierarchyTable_ElementType_Time", NewName="/Script/HierarchyTableAnimationRuntime.HierarchyTable_ElementType_Time")
++StructRedirects=(OldName="/Script/HierarchyTableBuiltinRuntime.HierarchyTable_TableType_Skeleton", NewName="/Script/HierarchyTableAnimationRuntime.HierarchyTable_TableType_Skeleton")
++StructRedirects=(OldName="/Script/HierarchyTableBuiltinRuntime.HierarchyTable_TablePayloadType_Skeleton", NewName="/Script/HierarchyTableAnimationRuntime.HierarchyTable_TablePayloadType_Skeleton")
\ No newline at end of file
diff --git a/Engine/Plugins/Experimental/Animation/HierarchyTableAnimation/HierarchyTableAnimation.uplugin b/Engine/Plugins/Experimental/Animation/HierarchyTableAnimation/HierarchyTableAnimation.uplugin
new file mode 100644
--- /dev/null
+++ b/Engine/Plugins/Experimental/Animation/HierarchyTableAnimation/HierarchyTableAnimation.uplugin
@@ -0,0 +1,45 @@
+{
+	"FileVersion": 3,
+	"Version": 1,
+	"VersionName": "0.1",
+	"FriendlyName": "Hierarchy Table Animation",
+	"Description": "Animation-specific type definitions for Hierarchy Tables",
+	"Category": "Animation",
+	"CreatedBy": "Epic Games, Inc.",
+	"CreatedByURL": "https://epicgames.com",
+	"DocsURL": "",
+	"MarketplaceURL": "",
+	"SupportURL": "",
+	"EnabledByDefault": false,
+	"CanContainContent": true,
+	"IsBetaVersion": false,
+	"IsExperimentalVersion": true,
+	"Installed": false,
+	"Modules": [
+		{
+			"Name": "HierarchyTableAnimationRuntime",
+			"Type": "Runtime",
+			"LoadingPhase": "Default"
+		},
+		{
+			"Name": "HierarchyTableAnimationEditor",
+			"Type": "Editor",
+			"LoadingPhase": "Default"
+		},
+		{
+			"Name": "HierarchyTableAnimationUncookedOnly",
+			"Type": "UncookedOnly",
+			"LoadingPhase": "Default"
+		}
+	],
+	"Plugins": [
+		{
+			"Name": "HierarchyTable",
+			"Enabled": true
+		},
+		{
+			"Name": "EditorDataStorageFeatures",
+			"Enabled": true
+		}
+	]
+}
diff --git a/Engine/Plugins/Experimental/Animation/HierarchyTableAnimation/Source/Editor/HierarchyTableAnimationEditor.Build.cs b/Engine/Plugins/Experimental/Animation/HierarchyTableAnimation/Source/Editor/HierarchyTableAnimationEditor.Build.cs
new file mode 100644
--- /dev/null
+++ b/Engine/Plugins/Experimental/Animation/HierarchyTableAnimation/Source/Editor/HierarchyTableAnimationEditor.Build.cs
@@ -0,0 +1,25 @@
+// Copyright Epic Games, Inc. All Rights Reserved.
+
+using EpicGames.Core;
+using UnrealBuildTool;
+
+public class HierarchyTableAnimationEditor : ModuleRules
+{
+	public HierarchyTableAnimationEditor(ReadOnlyTargetRules Target) : base(Target)
+	{
+		PrivateDependencyModuleNames.AddAll(
+			"Core",
+			"CoreUObject",
+			"Engine", 
+			"InputCore",
+			"Slate",
+			"SlateCore",
+			"UnrealEd",
+			"HierarchyTableRuntime",
+			"HierarchyTableEditor",
+			"HierarchyTableAnimationRuntime",
+			"TypedElementFramework", // TODO: Remove
+			"ToolMenus"
+		);
+	}
+}
diff --git a/Engine/Plugins/Experimental/Animation/HierarchyTableAnimation/Source/Editor/Private/HierachyTableAnimationEditorModule.cpp b/Engine/Plugins/Experimental/Animation/HierarchyTableAnimation/Source/Editor/Private/HierachyTableAnimationEditorModule.cpp
new file mode 100644
--- /dev/null
+++ b/Engine/Plugins/Experimental/Animation/HierarchyTableAnimation/Source/Editor/Private/HierachyTableAnimationEditorModule.cpp
@@ -0,0 +1,37 @@
+// Copyright Epic Games, Inc. All Rights Reserved.
+
+#include "HierarchyTableAnimationEditorModule.h"
+#include "TimeProfile/HierarchyTableTypeTime.h"
+#include "MaskProfile/HierarchyTableTypeMask.h"
+#include "Modules/ModuleManager.h"
+#include "HierarchyTableEditorModule.h"
+#include "SkeletonHierarchyTableType.h"
+#include "SkeletonHierarchyTableTypeHandler.h"
+#include "MaskProfile/MaskProfileProxyColumn.h"
+#include "TimeProfile/TimeProfileProxyColumn.h"
+
+void FHierarchyTableAnimationEditorModule::StartupModule()
+{
+	FHierarchyTableEditorModule& HierarchyTableModule = FModuleManager::LoadModuleChecked<FHierarchyTableEditorModule>("HierarchyTableEditor");
+	
+	HierarchyTableModule.RegisterTableType(FHierarchyTable_TableType_Skeleton::StaticStruct(), UHierarchyTable_TableTypeHandler_Skeleton::StaticClass());
+
+	HierarchyTableModule.RegisterElementTypeEditorColumns(FHierarchyTable_ElementType_Mask::StaticStruct(),
+		{ FHierarchyTableMaskColumn_Value::StaticStruct() });
+
+	HierarchyTableModule.RegisterElementTypeEditorColumns(FHierarchyTable_ElementType_Time::StaticStruct(),
+		{ FHierarchyTableTimeColumn_StartTime::StaticStruct(), FHierarchyTableTimeColumn_EndTime::StaticStruct(), FHierarchyTableTimeColumn_TimeFactor::StaticStruct() });
+}
+
+void FHierarchyTableAnimationEditorModule::ShutdownModule()
+{
+	if (FModuleManager::Get().IsModuleLoaded("HierarchyTableEditor"))
+	{
+		FHierarchyTableEditorModule& HierarchyTableModule = FModuleManager::GetModuleChecked<FHierarchyTableEditorModule>("HierarchyTableEditor");
+		HierarchyTableModule.UnregisterTableType(FHierarchyTable_TableType_Skeleton::StaticStruct());
+		HierarchyTableModule.UnregisterElementTypeEditorColumns(FHierarchyTable_ElementType_Mask::StaticStruct());
+		HierarchyTableModule.UnregisterElementTypeEditorColumns(FHierarchyTable_ElementType_Time::StaticStruct());
+	}
+}
+
+IMPLEMENT_MODULE(FHierarchyTableAnimationEditorModule, HierarchyTableAnimationEditor)
diff --git a/Engine/Plugins/Experimental/Animation/HierarchyTableAnimation/Source/Editor/Private/MaskProfile/MaskProfileFactory.cpp b/Engine/Plugins/Experimental/Animation/HierarchyTableAnimation/Source/Editor/Private/MaskProfile/MaskProfileFactory.cpp
new file mode 100644
--- /dev/null
+++ b/Engine/Plugins/Experimental/Animation/HierarchyTableAnimation/Source/Editor/Private/MaskProfile/MaskProfileFactory.cpp
@@ -0,0 +1,12 @@
+// Copyright Epic Games, Inc. All Rights Reserved.
+
+#include "MaskProfile/MaskProfileFactory.h"
+#include "MaskProfile/MaskProfileProxyColumn.h"
+#include "MaskProfile/MaskProfileWidgetConstructor.h"
+
+void UHierarchyTableMaskFactory::RegisterWidgetConstructors(IEditorDataStorageProvider& DataStorage, IEditorDataStorageUiProvider& DataStorageUi) const
+{
+	using namespace UE::Editor::DataStorage::Queries;
+
+	DataStorageUi.RegisterWidgetFactory<FHierarchyTableMaskWidgetConstructor_Value>(FName(TEXT("General.Cell")), TColumn<FHierarchyTableMaskColumn_Value>());
+}
\ No newline at end of file
diff --git a/Engine/Plugins/Experimental/Animation/HierarchyTableAnimation/Source/Editor/Private/MaskProfile/MaskProfileFactory.h b/Engine/Plugins/Experimental/Animation/HierarchyTableAnimation/Source/Editor/Private/MaskProfile/MaskProfileFactory.h
new file mode 100644
--- /dev/null
+++ b/Engine/Plugins/Experimental/Animation/HierarchyTableAnimation/Source/Editor/Private/MaskProfile/MaskProfileFactory.h
@@ -0,0 +1,16 @@
+// Copyright Epic Games, Inc. All Rights Reserved.
+
+#pragma once
+
+#include "Elements/Interfaces/TypedElementDataStorageFactory.h"
+
+#include "MaskProfileFactory.generated.h"
+
+UCLASS()
+class UHierarchyTableMaskFactory : public UEditorDataStorageFactory
+{
+	GENERATED_BODY()
+
+public:
+	void RegisterWidgetConstructors(IEditorDataStorageProvider& DataStorage, IEditorDataStorageUiProvider& DataStorageUi) const override;
+};
\ No newline at end of file
diff --git a/Engine/Plugins/Experimental/Animation/HierarchyTableAnimation/Source/Editor/Private/MaskProfile/MaskProfileWidgetConstructor.cpp b/Engine/Plugins/Experimental/Animation/HierarchyTableAnimation/Source/Editor/Private/MaskProfile/MaskProfileWidgetConstructor.cpp
new file mode 100644
--- /dev/null
+++ b/Engine/Plugins/Experimental/Animation/HierarchyTableAnimation/Source/Editor/Private/MaskProfile/MaskProfileWidgetConstructor.cpp
@@ -0,0 +1,37 @@
+// Copyright Epic Games, Inc. All Rights Reserved.
+
+#include "MaskProfile/MaskProfileWidgetConstructor.h"
+#include "Widgets/Input/SSpinBox.h"
+#include "MaskProfile/HierarchyTableTypeMask.h"
+#include "HierarchyTable.h"
+#include "Editor.h"
+
+#define LOCTEXT_NAMESPACE "FHierarchyTableMaskWidgetConstructor_Value"
+
+TSharedRef<SWidget> FHierarchyTableMaskWidgetConstructor_Value::CreateInternalWidget(UHierarchyTable* HierarchyTable, int32 EntryIndex)
+{
+	return SNew(SSpinBox<float>)
+		.IsEnabled_Lambda([HierarchyTable, EntryIndex]() { return HierarchyTable->GetTableEntry(EntryIndex)->IsOverridden(); })
+		.MinDesiredWidth(100.0f)
+		.MinValue(0.0f)
+		.MaxValue(1.0f)
+		.Value_Lambda([HierarchyTable, EntryIndex]()
+			{
+				return HierarchyTable->GetTableEntry(EntryIndex)->GetValue<FHierarchyTable_ElementType_Mask>()->Value;
+			})
+		.OnValueChanged_Lambda([HierarchyTable, EntryIndex](float NewValue)
+			{
+				HierarchyTable->GetMutableTableEntry(EntryIndex)->GetMutableValue<FHierarchyTable_ElementType_Mask>()->Value = NewValue;
+			})
+		.OnBeginSliderMovement_Lambda([HierarchyTable]()
+			{
+				GEditor->BeginTransaction(LOCTEXT("SetMaskValue", "Set Mask Value"));
+				HierarchyTable->Modify();
+			})
+		.OnEndSliderMovement_Lambda([](float)
+			{
+				GEditor->EndTransaction();
+			});
+}
+
+#undef LOCTEXT_NAMESPACE
\ No newline at end of file
diff --git a/Engine/Plugins/Experimental/Animation/HierarchyTableAnimation/Source/Editor/Private/MaskProfile/MaskProfileWidgetConstructor.h b/Engine/Plugins/Experimental/Animation/HierarchyTableAnimation/Source/Editor/Private/MaskProfile/MaskProfileWidgetConstructor.h
new file mode 100644
--- /dev/null
+++ b/Engine/Plugins/Experimental/Animation/HierarchyTableAnimation/Source/Editor/Private/MaskProfile/MaskProfileWidgetConstructor.h
@@ -0,0 +1,21 @@
+// Copyright Epic Games, Inc. All Rights Reserved.
+
+#pragma once
+
+#include "Elements/Interfaces/TypedElementDataStorageUiInterface.h"
+#include "HierarchyTableWidgetConstructor.h"
+
+#include "MaskProfileWidgetConstructor.generated.h"
+
+USTRUCT()
+struct FHierarchyTableMaskWidgetConstructor_Value : public FHierarchyTableWidgetConstructor
+{
+	GENERATED_BODY()
+
+public:
+	FHierarchyTableMaskWidgetConstructor_Value() : Super(StaticStruct()) {}
+	~FHierarchyTableMaskWidgetConstructor_Value() override = default;
+
+protected:
+	virtual TSharedRef<SWidget> CreateInternalWidget(UHierarchyTable* HierarchyTable, int32 EntryIndex) override;
+};
\ No newline at end of file
diff --git a/Engine/Plugins/Experimental/Animation/HierarchyTableAnimation/Source/Editor/Private/SkeletonHierarchyTableTypeHandler.cpp b/Engine/Plugins/Experimental/Animation/HierarchyTableAnimation/Source/Editor/Private/SkeletonHierarchyTableTypeHandler.cpp
new file mode 100644
--- /dev/null
+++ b/Engine/Plugins/Experimental/Animation/HierarchyTableAnimation/Source/Editor/Private/SkeletonHierarchyTableTypeHandler.cpp
@@ -0,0 +1,289 @@
+// Copyright Epic Games, Inc. All Rights Reserved.
+
+#include "SkeletonHierarchyTableTypeHandler.h"
+#include "ContentBrowserModule.h"
+#include "IContentBrowserSingleton.h"
+#include "SkeletonHierarchyTableType.h"
+#include "Animation/Skeleton.h"
+#include "Widgets/Layout/SBorder.h"
+#include "Editor.h"
+#include "StructUtils/InstancedStruct.h"
+#include "ToolMenu.h"
+#include "ToolMenuMisc.h"
+#include "Framework/Application/SlateApplication.h"
+#include "Widgets/Input/STextEntryPopup.h"
+#include "PersonaModule.h"
+#include "HierarchyTable.h"
+#include "HierarchyTableEditorModule.h"
+#include "Modules/ModuleManager.h"
+#include "Widgets/Images/SImage.h"
+#include "Widgets/Layout/SSpacer.h"
+#include "Widgets/Text/STextBlock.h"
+#include "Widgets/SBoxPanel.h"
+#include "IHierarchyTableEditor.h"
+
+#define LOCTEXT_NAMESPACE "UHierarchyTable_TableTypeHandler_Skeleton"
+
+void UHierarchyTable_TableTypeHandler_Skeleton::ConstructHierarchy()
+{
+	const FHierarchyTable_TableType_Skeleton& SkeletonTableType = HierarchyTable->GetTableMetadata<FHierarchyTable_TableType_Skeleton>();
+
+	check(SkeletonTableType.Skeleton);
+	check(HierarchyTable->GetTableData().Num() == 0);
+
+	FInstancedStruct DefaultEntry = HierarchyTable->CreateDefaultValue();
+
+	FInstancedStruct DefaultTablePayload;
+	DefaultTablePayload.InitializeAs<FHierarchyTable_TablePayloadType_Skeleton>();
+	DefaultTablePayload.GetMutable<FHierarchyTable_TablePayloadType_Skeleton>().EntryType = ESkeletonHierarchyTable_TablePayloadEntryType::Bone;
+
+	FReferenceSkeleton RefSkeleton = SkeletonTableType.Skeleton->GetReferenceSkeleton();
+
+	for (int32 BoneIndex = 0; BoneIndex < RefSkeleton.GetNum(); ++BoneIndex)
+	{
+		FHierarchyTableEntryData EntryData;
+		EntryData.Parent = RefSkeleton.GetParentIndex(BoneIndex);
+		EntryData.Identifier = RefSkeleton.GetBoneName(BoneIndex);
+		EntryData.TablePayload = DefaultTablePayload;
+		EntryData.Payload = BoneIndex == 0 ? DefaultEntry : TOptional<FInstancedStruct>();
+		EntryData.OwnerTable = HierarchyTable;
+		HierarchyTable->AddEntry(EntryData);
+	}
+}
+
+bool UHierarchyTable_TableTypeHandler_Skeleton::FactoryConfigureProperties(FInstancedStruct& TableType) const
+{
+	FHierarchyTable_TableType_Skeleton& SkeletonTableType = TableType.GetMutable<FHierarchyTable_TableType_Skeleton>();
+
+	FContentBrowserModule& ContentBrowserModule = FModuleManager::LoadModuleChecked<FContentBrowserModule>("ContentBrowser");
+
+	TSharedPtr<class SWindow> PickerWindow;
+
+	FAssetPickerConfig AssetPickerConfig;
+	AssetPickerConfig.Filter.ClassPaths.Add(USkeleton::StaticClass()->GetClassPathName());
+	AssetPickerConfig.InitialAssetViewType = EAssetViewType::List;
+	AssetPickerConfig.OnAssetSelected = FOnAssetSelected::CreateLambda([&](const FAssetData& SelectedAsset)
+		{
+			SkeletonTableType.Skeleton = Cast<USkeleton>(SelectedAsset.GetAsset());
+			PickerWindow->RequestDestroyWindow();
+		});
+
+	PickerWindow = SNew(SWindow)
+		.Title(INVTEXT("Pick Skeleton"))
+		.ClientSize(FVector2D(500, 600))
+		.SupportsMinimize(false)
+		.SupportsMaximize(false)
+		[
+			SNew(SBorder)
+			.BorderImage(FAppStyle::GetBrush("Menu.Background"))
+			[
+				ContentBrowserModule.Get().CreateAssetPicker(AssetPickerConfig)
+			]
+		];
+
+	GEditor->EditorAddModalWindow(PickerWindow.ToSharedRef());
+	PickerWindow.Reset();
+
+	return SkeletonTableType.Skeleton != nullptr;
+}
+
+void UHierarchyTable_TableTypeHandler_Skeleton::ExtendToolbar(UToolMenu* ToolMenu, const IHierarchyTableEditor* Editor) const
+{
+	const FToolMenuInsert SectionInsertLocation("Asset", EToolMenuInsertType::After);
+	{
+		FToolMenuSection& HierarchyTableSection = ToolMenu->AddSection("HierarchyTable", LOCTEXT("HierarchyTable_ToolbarLabel", "HierarchyTable"), SectionInsertLocation);
+
+		HierarchyTableSection.AddEntry(FToolMenuEntry::InitComboButton(
+			"AddCurve",
+			FUIAction(),
+			FNewToolMenuDelegate::CreateLambda([=, this](UToolMenu* InSubMenu)
+				{
+					FToolUIAction Action;
+				    Action.ExecuteAction = FToolMenuExecuteAction::CreateLambda([=, this](const FToolMenuContext& Context)
+				    {
+						TSharedRef<STextEntryPopup> TextEntry = SNew(STextEntryPopup)
+							.Label(LOCTEXT("NewCurveEntryLabal", "Curve Name"))
+							.OnTextCommitted_Lambda([=, this](const FText& CommittedText, ETextCommit::Type CommitInfo)
+								{
+									const int32 SelectedIndex = Editor->GetSelectedEntryIndex();
+									const int32 ParentIndex = SelectedIndex == INDEX_NONE ? 0 : SelectedIndex;
+									AddCurve(ParentIndex, FName(CommittedText.ToString()));
+									FSlateApplication::Get().DismissAllMenus();
+								});
+
+						FSlateApplication& SlateApp = FSlateApplication::Get();
+						SlateApp.PushMenu(
+							SlateApp.GetInteractiveTopLevelWindows()[0],
+							FWidgetPath(),
+							TextEntry,
+							SlateApp.GetCursorPos(),
+							FPopupTransitionEffect::TypeInPopup);
+				    });
+
+				    FToolMenuEntry Entry = FToolMenuEntry::InitMenuEntry(
+					    FName("AddNewCurve"),
+					    LOCTEXT("AddNewCurve_Label", "Add New Curve"),
+					    LOCTEXT("AddNewCurve_Tooltip", "Add a new curve value"),
+					    FSlateIcon(),
+					    FToolUIActionChoice(Action),
+						EUserInterfaceActionType::Button);
+				
+				    InSubMenu->AddMenuEntry("AddNewCurve", Entry);
+
+					const FHierarchyTable_TableType_Skeleton& SkeletonTableMetadata = HierarchyTable->GetTableMetadata<FHierarchyTable_TableType_Skeleton>();
+
+					FPersonaModule& PersonaModule = FModuleManager::LoadModuleChecked<FPersonaModule>("Persona");
+					InSubMenu->AddMenuEntry("ExistingCurveMenu", FToolMenuEntry::InitWidget(
+						"ExistingCurveMenu",
+						SNew(SVerticalBox)
+						+SVerticalBox::Slot()
+						.AutoHeight()
+						[
+							PersonaModule.CreateCurvePicker(SkeletonTableMetadata.Skeleton,
+								FOnCurvePicked::CreateLambda([&](const FName& InName)
+									{
+										const int32 SelectedIndex = Editor->GetSelectedEntryIndex();
+										const int32 ParentIndex = SelectedIndex == INDEX_NONE ? 0 : SelectedIndex;
+										AddCurve(ParentIndex, InName);
+										FSlateApplication::Get().DismissAllMenus();
+									}))
+						],
+						FText(),
+						true,
+						false,
+						true
+					));
+				}),
+			LOCTEXT("AddCurve_Label", "Add Curve"),
+			LOCTEXT("AddCurve_ToolTip", "Add a new curve to the hierarchy"),
+			FSlateIcon(FAppStyle::Get().GetStyleSetName(), "Icons.Plus")
+		));
+
+		HierarchyTableSection.AddEntry(FToolMenuEntry::InitComboButton(
+			"AddAttribute",
+			FUIAction(),
+			FNewToolMenuDelegate::CreateLambda([=, this](UToolMenu* InSubMenu)
+				{
+					FToolUIAction Action;
+				    Action.ExecuteAction = FToolMenuExecuteAction::CreateLambda([=, this](const FToolMenuContext& Context)
+				    {
+						TSharedRef<STextEntryPopup> TextEntry = SNew(STextEntryPopup)
+							.Label(LOCTEXT("NewAttributeEntryLabal", "Attribute Name"))
+							.OnTextCommitted_Lambda([=, this](const FText& CommittedText, ETextCommit::Type CommitInfo)
+								{
+									const int32 SelectedIndex = Editor->GetSelectedEntryIndex();
+									const int32 ParentIndex = SelectedIndex == INDEX_NONE ? 0 : SelectedIndex;
+									AddAttribute(ParentIndex, FName(CommittedText.ToString()));
+									FSlateApplication::Get().DismissAllMenus();
+								});
+
+						FSlateApplication& SlateApp = FSlateApplication::Get();
+						SlateApp.PushMenu(
+							SlateApp.GetInteractiveTopLevelWindows()[0],
+							FWidgetPath(),
+							TextEntry,
+							SlateApp.GetCursorPos(),
+							FPopupTransitionEffect::TypeInPopup);
+				    });
+
+				    FToolMenuEntry Entry = FToolMenuEntry::InitMenuEntry(
+					    FName("AddNewAttribute"),
+					    LOCTEXT("AddNewAttribute_Label", "Add New Attribute"),
+					    LOCTEXT("AddNewAttribute_Tooltip", "Add a new attribute value"),
+					    FSlateIcon(),
+					    FToolUIActionChoice(Action),
+						EUserInterfaceActionType::Button);
+				
+				    InSubMenu->AddMenuEntry("AddNewAttribute", Entry);
+				}),
+			LOCTEXT("AddAttribute_Label", "Add Attribute"),
+			LOCTEXT("AddAttribute_ToolTip", "Add a new attribute to the hierarchy"),
+			FSlateIcon(FAppStyle::Get().GetStyleSetName(), "Icons.Plus")
+		));
+	}
+}
+
+TSharedRef<SWidget> UHierarchyTable_TableTypeHandler_Skeleton::CreateLabelWidget(const int32 EntryIndex) const
+{
+	const FSlateBrush* LabelIcon = nullptr;
+	FSlateColor LabelIconColor = FSlateColor::UseForeground();
+
+	const FHierarchyTableEntryData* const Entry = HierarchyTable->GetTableEntry(EntryIndex);
+	const FHierarchyTable_TablePayloadType_Skeleton& EntryTablePayload = Entry->TablePayload.Get<FHierarchyTable_TablePayloadType_Skeleton>();
+
+	switch (EntryTablePayload.EntryType)
+	{
+		case ESkeletonHierarchyTable_TablePayloadEntryType::Bone:
+			LabelIcon = FAppStyle::GetBrush("SkeletonTree.Bone");
+			break;
+		case ESkeletonHierarchyTable_TablePayloadEntryType::Curve:
+			LabelIcon = FAppStyle::GetBrush("AnimGraph.Attribute.Curves.Icon");
+			LabelIconColor = FAppStyle::GetSlateColor("AnimGraph.Attribute.Curves.Color");
+			break;
+		case ESkeletonHierarchyTable_TablePayloadEntryType::Attribute:
+			LabelIcon = FAppStyle::GetBrush("AnimGraph.Attribute.Attributes.Icon");
+			LabelIconColor = FAppStyle::GetSlateColor("AnimGraph.Attribute.Attributes.Color");
+			break;
+	}
+
+	return SNew(SHorizontalBox)
+		+SHorizontalBox::Slot()
+		.AutoWidth()
+		[
+			SNew(SImage)
+			.Image(LabelIcon)
+			.ColorAndOpacity(LabelIconColor)
+		]
+		+SHorizontalBox::Slot()
+		.AutoWidth()
+		[
+			SNew(SSpacer)
+			.Size(FVector2D(5.0f, 0.0f))
+		]
+		+SHorizontalBox::Slot()
+		.FillWidth(1.0f)
+		[
+			SNew(STextBlock)
+			.Text(FText::FromName(Entry->Identifier))
+		];
+}
+
+void UHierarchyTable_TableTypeHandler_Skeleton::AddCurve(const int32 ParentIndex, const FName Identifier) const
+{
+	FHierarchyTableEntryData NewEntry;
+	NewEntry.Parent = ParentIndex;
+	NewEntry.Identifier = Identifier;
+	NewEntry.Payload = TOptional<FInstancedStruct>();
+	NewEntry.OwnerTable = HierarchyTable;
+	{
+		FInstancedStruct TablePayload;
+		TablePayload.InitializeAs<FHierarchyTable_TablePayloadType_Skeleton>();
+		FHierarchyTable_TablePayloadType_Skeleton& SkeletonTablePayload = TablePayload.GetMutable<FHierarchyTable_TablePayloadType_Skeleton>();
+		SkeletonTablePayload.EntryType = ESkeletonHierarchyTable_TablePayloadEntryType::Curve;
+		
+		NewEntry.TablePayload = TablePayload;
+	}
+
+	HierarchyTable->AddEntry(NewEntry);
+}
+	
+void UHierarchyTable_TableTypeHandler_Skeleton::AddAttribute(const int32 ParentIndex, const FName Identifier) const
+{
+	FHierarchyTableEntryData NewEntry;
+	NewEntry.Parent = ParentIndex;
+	NewEntry.Identifier = Identifier;
+	NewEntry.Payload = TOptional<FInstancedStruct>();
+	NewEntry.OwnerTable = HierarchyTable;
+	{
+		FInstancedStruct TablePayload;
+		TablePayload.InitializeAs<FHierarchyTable_TablePayloadType_Skeleton>();
+		FHierarchyTable_TablePayloadType_Skeleton& SkeletonTablePayload = TablePayload.GetMutable<FHierarchyTable_TablePayloadType_Skeleton>();
+		SkeletonTablePayload.EntryType = ESkeletonHierarchyTable_TablePayloadEntryType::Attribute;
+		
+		NewEntry.TablePayload = TablePayload;
+	}
+
+	HierarchyTable->AddEntry(NewEntry);
+}
+
+#undef LOCTEXT_NAMESPACE
\ No newline at end of file
diff --git a/Engine/Plugins/Experimental/Animation/HierarchyTableAnimation/Source/Editor/Private/SkeletonHierarchyTableTypeHandler.h b/Engine/Plugins/Experimental/Animation/HierarchyTableAnimation/Source/Editor/Private/SkeletonHierarchyTableTypeHandler.h
new file mode 100644
--- /dev/null
+++ b/Engine/Plugins/Experimental/Animation/HierarchyTableAnimation/Source/Editor/Private/SkeletonHierarchyTableTypeHandler.h
@@ -0,0 +1,27 @@
+// Copyright Epic Games, Inc. All Rights Reserved.
+
+#pragma once
+
+#include "HierarchyTableTypeHandler.h"
+
+#include "SkeletonHierarchyTableTypeHandler.generated.h"
+
+UCLASS()
+class UHierarchyTable_TableTypeHandler_Skeleton final : public UHierarchyTable_TableTypeHandler
+{
+	GENERATED_BODY()
+
+public:
+	virtual void ConstructHierarchy();
+
+	virtual bool FactoryConfigureProperties(FInstancedStruct& TableType) const override;
+
+	virtual void ExtendToolbar(UToolMenu* ToolMenu, const IHierarchyTableEditor* Editor) const override;
+
+	virtual TSharedRef<SWidget> CreateLabelWidget(const int32 EntryIndex) const;
+
+private:
+	void AddCurve(const int32 ParentIndex, const FName Identifier) const;
+	
+	void AddAttribute(const int32 ParentIndex, const FName Identifier) const;
+};
\ No newline at end of file
diff --git a/Engine/Plugins/Experimental/Animation/HierarchyTableAnimation/Source/Editor/Private/TimeProfile/TimeProfileFactory.cpp b/Engine/Plugins/Experimental/Animation/HierarchyTableAnimation/Source/Editor/Private/TimeProfile/TimeProfileFactory.cpp
new file mode 100644
--- /dev/null
+++ b/Engine/Plugins/Experimental/Animation/HierarchyTableAnimation/Source/Editor/Private/TimeProfile/TimeProfileFactory.cpp
@@ -0,0 +1,15 @@
+// Copyright Epic Games, Inc. All Rights Reserved.
+
+#include "TimeProfile/TimeProfileFactory.h"
+#include "TimeProfile/TimeProfileProxyColumn.h"
+#include "TimeProfile/TimeProfileWidgetConstructor.h"
+
+void UHierarchyTableTimeFactory::RegisterWidgetConstructors(IEditorDataStorageProvider& DataStorage, IEditorDataStorageUiProvider& DataStorageUi) const
+{
+	using namespace UE::Editor::DataStorage::Queries;
+
+	DataStorageUi.RegisterWidgetFactory<FHierarchyTableTimeWidgetConstructor_StartTime>(FName(TEXT("General.Cell")), TColumn<FHierarchyTableTimeColumn_StartTime>());
+	DataStorageUi.RegisterWidgetFactory<FHierarchyTableTimeWidgetConstructor_EndTime>(FName(TEXT("General.Cell")), TColumn<FHierarchyTableTimeColumn_EndTime>());
+	DataStorageUi.RegisterWidgetFactory<FHierarchyTableTimeWidgetConstructor_TimeFactor>(FName(TEXT("General.Cell")), TColumn<FHierarchyTableTimeColumn_TimeFactor>());
+	DataStorageUi.RegisterWidgetFactory<FHierarchyTableTimeWidgetConstructor_Preview>(FName(TEXT("General.Cell")), TColumn<FHierarchyTableTimeColumn_Preview>());
+}
\ No newline at end of file
diff --git a/Engine/Plugins/Experimental/Animation/HierarchyTableAnimation/Source/Editor/Private/TimeProfile/TimeProfileFactory.h b/Engine/Plugins/Experimental/Animation/HierarchyTableAnimation/Source/Editor/Private/TimeProfile/TimeProfileFactory.h
new file mode 100644
--- /dev/null
+++ b/Engine/Plugins/Experimental/Animation/HierarchyTableAnimation/Source/Editor/Private/TimeProfile/TimeProfileFactory.h
@@ -0,0 +1,15 @@
+// Copyright Epic Games, Inc. All Rights Reserved.
+
+#pragma once
+
+#include "Elements/Interfaces/TypedElementDataStorageFactory.h"
+#include "TimeProfileFactory.generated.h"
+
+UCLASS()
+class UHierarchyTableTimeFactory : public UEditorDataStorageFactory
+{
+	GENERATED_BODY()
+
+public:
+	void RegisterWidgetConstructors(IEditorDataStorageProvider& DataStorage, IEditorDataStorageUiProvider& DataStorageUi) const override;
+};
\ No newline at end of file
diff --git a/Engine/Plugins/Experimental/Animation/HierarchyTableAnimation/Source/Editor/Private/TimeProfile/TimeProfileWidgetConstructor.cpp b/Engine/Plugins/Experimental/Animation/HierarchyTableAnimation/Source/Editor/Private/TimeProfile/TimeProfileWidgetConstructor.cpp
new file mode 100644
--- /dev/null
+++ b/Engine/Plugins/Experimental/Animation/HierarchyTableAnimation/Source/Editor/Private/TimeProfile/TimeProfileWidgetConstructor.cpp
@@ -0,0 +1,63 @@
+// Copyright Epic Games, Inc. All Rights Reserved.
+
+#include "TimeProfileWidgetConstructor.h"
+#include "Widgets/Input/SSpinBox.h"
+#include "TimeProfile/HierarchyTableTypeTime.h"
+#include "HierarchyTable.h"
+
+TSharedRef<SWidget> FHierarchyTableTimeWidgetConstructor_StartTime::CreateInternalWidget(UHierarchyTable* HierarchyTable, int32 EntryIndex)
+{
+	return SNew(SSpinBox<float>)
+		.IsEnabled_Lambda([HierarchyTable, EntryIndex]() { return HierarchyTable->GetTableEntry(EntryIndex)->IsOverridden(); })
+		.MinDesiredWidth(100.0f)
+		.MinValue(0.0f)
+		.MaxValue(1.0f)
+		.Value_Lambda([HierarchyTable, EntryIndex]()
+			{
+				return HierarchyTable->GetTableEntry(EntryIndex)->GetValue<FHierarchyTable_ElementType_Time>()->StartTime;
+			})
+		.OnValueChanged_Lambda([HierarchyTable, EntryIndex](float NewValue)
+			{
+				HierarchyTable->GetMutableTableEntry(EntryIndex)->GetMutableValue<FHierarchyTable_ElementType_Time>()->StartTime = NewValue;
+			});
+}
+
+TSharedRef<SWidget> FHierarchyTableTimeWidgetConstructor_EndTime::CreateInternalWidget(UHierarchyTable* HierarchyTable, int32 EntryIndex)
+{
+	return SNew(SSpinBox<float>)
+		.IsEnabled_Lambda([HierarchyTable, EntryIndex]() { return HierarchyTable->GetTableEntry(EntryIndex)->IsOverridden(); })
+		.MinDesiredWidth(100.0f)
+		.MinValue(0.0f)
+		.MaxValue(1.0f)
+		.Value_Lambda([HierarchyTable, EntryIndex]()
+			{
+				return HierarchyTable->GetTableEntry(EntryIndex)->GetValue<FHierarchyTable_ElementType_Time>()->EndTime;
+			})
+		.OnValueChanged_Lambda([HierarchyTable, EntryIndex](float NewValue)
+			{
+				HierarchyTable->GetMutableTableEntry(EntryIndex)->GetMutableValue<FHierarchyTable_ElementType_Time>()->EndTime = NewValue;
+			});
+}
+
+TSharedRef<SWidget> FHierarchyTableTimeWidgetConstructor_TimeFactor::CreateInternalWidget(UHierarchyTable* HierarchyTable, int32 EntryIndex)
+{
+	return SNew(SSpinBox<float>)
+		.IsEnabled_Lambda([HierarchyTable, EntryIndex]() { return HierarchyTable->GetTableEntry(EntryIndex)->IsOverridden(); })
+		.MinDesiredWidth(100.0f)
+		.MinValue(0.0f)
+		.MaxValue(1.0f)
+		.Value_Lambda([HierarchyTable, EntryIndex]()
+			{
+				return HierarchyTable->GetTableEntry(EntryIndex)->GetValue<FHierarchyTable_ElementType_Time>()->TimeFactor;
+			})
+		.OnValueChanged_Lambda([HierarchyTable, EntryIndex](float NewValue)
+			{
+				HierarchyTable->GetMutableTableEntry(EntryIndex)->GetMutableValue<FHierarchyTable_ElementType_Time>()->TimeFactor = NewValue;
+			});
+}
+
+TSharedRef<SWidget> FHierarchyTableTimeWidgetConstructor_Preview::CreateInternalWidget(UHierarchyTable* HierarchyTable, int32 EntryIndex)
+{
+	return SNew(STextBlock)
+		.Text(INVTEXT("PREVIEW"));
+}
\ No newline at end of file
diff --git a/Engine/Plugins/Experimental/Animation/HierarchyTableAnimation/Source/Editor/Private/TimeProfile/TimeProfileWidgetConstructor.h b/Engine/Plugins/Experimental/Animation/HierarchyTableAnimation/Source/Editor/Private/TimeProfile/TimeProfileWidgetConstructor.h
new file mode 100644
--- /dev/null
+++ b/Engine/Plugins/Experimental/Animation/HierarchyTableAnimation/Source/Editor/Private/TimeProfile/TimeProfileWidgetConstructor.h
@@ -0,0 +1,60 @@
+// Copyright Epic Games, Inc. All Rights Reserved.
+
+#pragma once
+
+#include "Elements/Interfaces/TypedElementDataStorageUiInterface.h"
+#include "HierarchyTableWidgetConstructor.h"
+
+#include "TimeProfileWidgetConstructor.generated.h"
+
+USTRUCT()
+struct FHierarchyTableTimeWidgetConstructor_StartTime : public FHierarchyTableWidgetConstructor
+{
+	GENERATED_BODY()
+
+public:
+	FHierarchyTableTimeWidgetConstructor_StartTime() : Super(StaticStruct()) {}
+	~FHierarchyTableTimeWidgetConstructor_StartTime() override = default;
+
+protected:
+	virtual TSharedRef<SWidget> CreateInternalWidget(UHierarchyTable* HierarchyTable, int32 EntryIndex) override;
+};
+
+USTRUCT()
+struct FHierarchyTableTimeWidgetConstructor_EndTime : public FHierarchyTableWidgetConstructor
+{
+	GENERATED_BODY()
+
+public:
+	FHierarchyTableTimeWidgetConstructor_EndTime() : Super(StaticStruct()) {}
+	~FHierarchyTableTimeWidgetConstructor_EndTime() override = default;
+
+protected:
+	virtual TSharedRef<SWidget> CreateInternalWidget(UHierarchyTable* HierarchyTable, int32 EntryIndex) override;
+};
+
+USTRUCT()
+struct FHierarchyTableTimeWidgetConstructor_TimeFactor : public FHierarchyTableWidgetConstructor
+{
+	GENERATED_BODY()
+
+public:
+	FHierarchyTableTimeWidgetConstructor_TimeFactor() : Super(StaticStruct()) {}
+	~FHierarchyTableTimeWidgetConstructor_TimeFactor() override = default;
+
+protected:
+	virtual TSharedRef<SWidget> CreateInternalWidget(UHierarchyTable* HierarchyTable, int32 EntryIndex) override;
+};
+
+USTRUCT()
+struct FHierarchyTableTimeWidgetConstructor_Preview : public FHierarchyTableWidgetConstructor
+{
+	GENERATED_BODY()
+
+public:
+	FHierarchyTableTimeWidgetConstructor_Preview() : Super(StaticStruct()) {}
+	~FHierarchyTableTimeWidgetConstructor_Preview() override = default;
+
+protected:
+	virtual TSharedRef<SWidget> CreateInternalWidget(UHierarchyTable* HierarchyTable, int32 EntryIndex) override;
+};
\ No newline at end of file
diff --git a/Engine/Plugins/Experimental/Animation/HierarchyTableAnimation/Source/Editor/Public/HierarchyTableAnimationEditorModule.h b/Engine/Plugins/Experimental/Animation/HierarchyTableAnimation/Source/Editor/Public/HierarchyTableAnimationEditorModule.h
new file mode 100644
--- /dev/null
+++ b/Engine/Plugins/Experimental/Animation/HierarchyTableAnimation/Source/Editor/Public/HierarchyTableAnimationEditorModule.h
@@ -0,0 +1,13 @@
+// Copyright Epic Games, Inc. All Rights Reserved.
+
+#pragma once
+
+#include "Modules/ModuleInterface.h"
+
+class FHierarchyTableAnimationEditorModule : public IModuleInterface
+{
+public:
+	virtual void StartupModule() override;
+	
+	virtual void ShutdownModule() override;
+};
\ No newline at end of file
diff --git a/Engine/Plugins/Experimental/Animation/HierarchyTableAnimation/Source/Editor/Public/MaskProfile/MaskProfileProxyColumn.h b/Engine/Plugins/Experimental/Animation/HierarchyTableAnimation/Source/Editor/Public/MaskProfile/MaskProfileProxyColumn.h
new file mode 100644
--- /dev/null
+++ b/Engine/Plugins/Experimental/Animation/HierarchyTableAnimation/Source/Editor/Public/MaskProfile/MaskProfileProxyColumn.h
@@ -0,0 +1,13 @@
+// Copyright Epic Games, Inc. All Rights Reserved.
+
+#pragma once
+
+#include "Elements/Common/TypedElementCommonTypes.h"
+
+#include "MaskProfileProxyColumn.generated.h"
+
+USTRUCT(meta = (DisplayName = "Mask"))
+struct FHierarchyTableMaskColumn_Value final : public FEditorDataStorageColumn
+{
+	GENERATED_BODY()
+};
\ No newline at end of file
diff --git a/Engine/Plugins/Experimental/Animation/HierarchyTableAnimation/Source/Editor/Public/TimeProfile/TimeProfileProxyColumn.h b/Engine/Plugins/Experimental/Animation/HierarchyTableAnimation/Source/Editor/Public/TimeProfile/TimeProfileProxyColumn.h
new file mode 100644
--- /dev/null
+++ b/Engine/Plugins/Experimental/Animation/HierarchyTableAnimation/Source/Editor/Public/TimeProfile/TimeProfileProxyColumn.h
@@ -0,0 +1,31 @@
+// Copyright Epic Games, Inc. All Rights Reserved.
+
+#pragma once
+
+#include "Elements/Common/TypedElementCommonTypes.h"
+
+#include "TimeProfileProxyColumn.generated.h"
+
+USTRUCT()
+struct FHierarchyTableTimeColumn_StartTime final : public FEditorDataStorageColumn
+{
+	GENERATED_BODY()
+};
+
+USTRUCT()
+struct FHierarchyTableTimeColumn_EndTime final : public FEditorDataStorageColumn
+{
+	GENERATED_BODY()
+};
+
+USTRUCT()
+struct FHierarchyTableTimeColumn_TimeFactor final : public FEditorDataStorageColumn
+{
+	GENERATED_BODY()
+};
+
+USTRUCT()
+struct FHierarchyTableTimeColumn_Preview final : public FEditorDataStorageColumn
+{
+	GENERATED_BODY()
+};
\ No newline at end of file
diff --git a/Engine/Plugins/Experimental/Animation/HierarchyTableAnimation/Source/Runtime/HierarchyTableAnimationRuntime.Build.cs b/Engine/Plugins/Experimental/Animation/HierarchyTableAnimation/Source/Runtime/HierarchyTableAnimationRuntime.Build.cs
new file mode 100644
--- /dev/null
+++ b/Engine/Plugins/Experimental/Animation/HierarchyTableAnimation/Source/Runtime/HierarchyTableAnimationRuntime.Build.cs
@@ -0,0 +1,25 @@
+// Copyright Epic Games, Inc. All Rights Reserved.
+
+using UnrealBuildTool;
+
+public class HierarchyTableAnimationRuntime : ModuleRules
+{
+	public HierarchyTableAnimationRuntime(ReadOnlyTargetRules Target) : base(Target)
+	{
+		PublicDependencyModuleNames.AddRange(
+			new string[]
+			{
+				"Core",
+				"CoreUObject",
+				"Engine",
+				"HierarchyTableRuntime"
+			}
+		);
+
+		PrivateDependencyModuleNames.AddRange(
+			new string[]
+            {
+			}
+		);
+	}
+}
diff --git a/Engine/Plugins/Experimental/Animation/HierarchyTableAnimation/Source/Runtime/Private/AnimNode_BlendProfileLayeredBlend.cpp b/Engine/Plugins/Experimental/Animation/HierarchyTableAnimation/Source/Runtime/Private/AnimNode_BlendProfileLayeredBlend.cpp
new file mode 100644
--- /dev/null
+++ b/Engine/Plugins/Experimental/Animation/HierarchyTableAnimation/Source/Runtime/Private/AnimNode_BlendProfileLayeredBlend.cpp
@@ -0,0 +1,346 @@
+// Copyright Epic Games, Inc. All Rights Reserved.
+
+#include "AnimNode_BlendProfileLayeredBlend.h"
+#include "AnimationRuntime.h"
+#include "Animation/AnimInstanceProxy.h"
+#include "Animation/AnimTrace.h"
+#include "Animation/AnimCurveTypes.h"
+#include "Animation/AnimCurveUtils.h"
+#include "HierarchyTable.h"
+#include "MaskProfile/HierarchyTableTypeMask.h"
+#include "Animation/AttributeTypes.h"
+#include "Animation/IAttributeBlendOperator.h"
+#include "Hash/CityHash.h"
+#include "SkeletonHierarchyTableType.h"
+
+#include UE_INLINE_GENERATED_CPP_BY_NAME(AnimNode_BlendProfileLayeredBlend)
+
+void FAnimNode_BlendProfileLayeredBlend::Initialize_AnyThread(const FAnimationInitializeContext& Context)
+{
+	DECLARE_SCOPE_HIERARCHICAL_COUNTER_ANIMNODE(Initialize_AnyThread)
+	FAnimNode_Base::Initialize_AnyThread(Context);
+
+	BasePose.Initialize(Context);
+	BlendPose.Initialize(Context);
+}
+
+
+void FAnimNode_BlendProfileLayeredBlend::RebuildPerBoneBlendWeights(const USkeleton* InSkeleton)
+{
+	if (InSkeleton)
+	{
+		CreateMaskWeights(InSkeleton);
+
+		SkeletonGuid = InSkeleton->GetGuid();
+		VirtualBoneGuid = InSkeleton->GetVirtualBoneGuid();
+	}
+}
+
+bool FAnimNode_BlendProfileLayeredBlend::ArePerBoneBlendWeightsValid(const USkeleton* InSkeleton) const
+{
+	return (InSkeleton != nullptr && InSkeleton->GetGuid() == SkeletonGuid && InSkeleton->GetVirtualBoneGuid() == VirtualBoneGuid);
+}
+
+void FAnimNode_BlendProfileLayeredBlend::UpdateCachedBoneData(const FBoneContainer& RequiredBones, const USkeleton* Skeleton)
+{
+	if (RequiredBones.GetSerialNumber() == RequiredBonesSerialNumber)
+	{
+		return;
+	}
+	
+	// Update cached mask weights
+	{
+		CachedCurveMaskWeights.Empty();
+		//CachedAttributeMaskWeights.Empty();
+
+		for (const FHierarchyTableEntryData& Entry : MaskTable->GetTableData())
+		{
+			if (Entry.TablePayload.Get<FHierarchyTable_TablePayloadType_Skeleton>().EntryType == ESkeletonHierarchyTable_TablePayloadEntryType::Curve)
+			{
+				const float EntryWeight = Entry.GetValue<FHierarchyTable_ElementType_Mask>()->Value;
+				CachedCurveMaskWeights.Add<UE::Anim::FCurveElement>({ Entry.Identifier, EntryWeight });
+			}
+			/*
+			else if (Entry.EntryType == EHierarchyTableEntryType::Attribute)
+			{
+				const float EntryWeight = Entry.GetValue<FHierarchyTableType_Mask>()->Value;
+				CachedAttributeMaskWeights.Add(FNamedFloat(Entry.Identifier, EntryWeight));
+			}
+			*/
+		}
+	}
+
+	if (!ArePerBoneBlendWeightsValid(Skeleton))
+	{
+		RebuildPerBoneBlendWeights(Skeleton);
+	}
+
+	// build desired bone weights
+	const TArray<FBoneIndexType>& RequiredBoneIndices = RequiredBones.GetBoneIndicesArray();
+	const int32 NumRequiredBones = RequiredBoneIndices.Num();
+	DesiredBoneBlendWeights.SetNumZeroed(NumRequiredBones);
+	for (int32 RequiredBoneIndex = 0; RequiredBoneIndex < NumRequiredBones; RequiredBoneIndex++)
+	{
+		const int32 SkeletonBoneIndex = RequiredBones.GetSkeletonIndex(FCompactPoseBoneIndex(RequiredBoneIndex));
+		if (ensure(SkeletonBoneIndex != INDEX_NONE))
+		{
+			DesiredBoneBlendWeights[RequiredBoneIndex] = PerBoneBlendWeights[SkeletonBoneIndex];
+		}
+	}
+
+	CurrentBoneBlendWeights.Reset(DesiredBoneBlendWeights.Num());
+	CurrentBoneBlendWeights.AddZeroed(DesiredBoneBlendWeights.Num());
+
+	//Reinitialize bone blend weights now that we have cleared them
+	UpdateDesiredBoneWeight();
+
+	RequiredBonesSerialNumber = RequiredBones.GetSerialNumber();
+}
+
+void FAnimNode_BlendProfileLayeredBlend::CacheBones_AnyThread(const FAnimationCacheBonesContext& Context)
+{
+	DECLARE_SCOPE_HIERARCHICAL_COUNTER_ANIMNODE(CacheBones_AnyThread)
+	
+	BasePose.CacheBones(Context);
+	BlendPose.CacheBones(Context);
+
+	UpdateCachedBoneData(Context.AnimInstanceProxy->GetRequiredBones(), Context.AnimInstanceProxy->GetSkeleton());
+}
+
+void FAnimNode_BlendProfileLayeredBlend::Update_AnyThread(const FAnimationUpdateContext& Context)
+{
+	DECLARE_SCOPE_HIERARCHICAL_COUNTER_ANIMNODE(Update_AnyThread)
+	bool RootMotionBlendPose = false;
+	float RootMotionWeight = 0.f;
+	const float RootMotionClearWeight = bBlendRootMotionBasedOnRootBone ? 0.f : 1.f;
+
+	if (IsLODEnabled(Context.AnimInstanceProxy))
+	{
+		GetEvaluateGraphExposedInputs().Execute(Context);
+
+		if (FAnimWeight::IsRelevant(BlendWeight))
+		{
+			UpdateCachedBoneData(Context.AnimInstanceProxy->GetRequiredBones(), Context.AnimInstanceProxy->GetSkeleton());
+			UpdateDesiredBoneWeight();
+
+			if (bBlendRootMotionBasedOnRootBone && !CurrentBoneBlendWeights.IsEmpty())
+			{
+				const float NewRootMotionWeight = CurrentBoneBlendWeights[0];
+				if (NewRootMotionWeight > ZERO_ANIMWEIGHT_THRESH)
+				{
+					RootMotionWeight = NewRootMotionWeight;
+					RootMotionBlendPose = true;
+				}
+			}
+
+			const float ThisPoseRootMotionWeight = RootMotionBlendPose ? RootMotionWeight : RootMotionClearWeight;
+			BlendPose.Update(Context.FractionalWeightAndRootMotion(BlendWeight, ThisPoseRootMotionWeight));
+		}
+	}
+
+	// initialize children
+	const float BaseRootMotionWeight = 1.f - RootMotionWeight;
+
+	if (BaseRootMotionWeight < ZERO_ANIMWEIGHT_THRESH)
+	{
+		BasePose.Update(Context.FractionalWeightAndRootMotion(1.f, BaseRootMotionWeight));
+	}
+	else
+	{
+		BasePose.Update(Context);
+	}
+}
+
+void FAnimNode_BlendProfileLayeredBlend::Evaluate_AnyThread(FPoseContext& Output)
+{
+	DECLARE_SCOPE_HIERARCHICAL_COUNTER_ANIMNODE(Evaluate_AnyThread)
+	ANIM_MT_SCOPE_CYCLE_COUNTER(BlendPosesInGraph, !IsInGameThread());
+
+	FPoseContext BasePoseContext(Output);
+	BasePose.Evaluate(BasePoseContext);
+
+	FCompactPose TargetPose;
+	FBlendedCurve TargetCurve;
+	UE::Anim::FStackAttributeContainer TargetAttributes;
+
+	if (FAnimWeight::IsRelevant(BlendWeight))
+	{
+		FPoseContext CurrentPoseContext(Output);
+		BlendPose.Evaluate(CurrentPoseContext);
+
+		TargetPose.MoveBonesFrom(CurrentPoseContext.Pose);
+		TargetCurve.MoveFrom(CurrentPoseContext.Curve);
+		TargetAttributes.MoveFrom(CurrentPoseContext.CustomAttributes);
+	}
+	else
+	{
+		TargetPose.ResetToRefPose(BasePoseContext.Pose.GetBoneContainer());
+		TargetCurve.InitFrom(Output.Curve);
+		TargetAttributes.CopyFrom(Output.CustomAttributes);
+	}
+
+	FAnimationPoseData OutputPoseData(Output);
+	
+	// Blend poses
+	FAnimationRuntime::BlendTwoPosesTogetherPerBone(BasePoseContext.Pose, TargetPose, CurrentBoneBlendWeights, OutputPoseData.GetPose());
+
+	// Blend curves
+	{
+		FBlendedCurve& OutCurve = OutputPoseData.GetCurve();
+		OutCurve.CopyFrom(BasePoseContext.Curve);
+
+		if (FAnimWeight::IsRelevant(BlendWeight))
+		{
+			FBlendedCurve FilteredCurves;
+
+			// Multiply per-curve blend weights by matching blend pose curves
+			UE::Anim::FNamedValueArrayUtils::Intersection(
+				TargetCurve,
+				CachedCurveMaskWeights,
+				[this, &FilteredCurves](const UE::Anim::FCurveElement& InBlendElement, const UE::Anim::FCurveElement& InMaskElement) mutable
+					{
+						FilteredCurves.Add(InBlendElement.Name, InBlendElement.Value * InMaskElement.Value);
+					});
+
+			// Override blend curve values with premultipled curves
+			TargetCurve.Combine(FilteredCurves);
+
+			// Remove curves that have been filtered by the mask, curves with no mask value defined remain, even with 0.0 value
+			UE::Anim::FNamedValueArrayUtils::RemoveByPredicate(
+				TargetCurve,
+				CachedCurveMaskWeights,
+				[](const UE::Anim::FCurveElement& InBaseElement, const UE::Anim::FCurveElement& InMaskElement)
+				{	
+					const bool bKeep = InMaskElement.Value == 0.0;
+					return bKeep;
+				}
+			);
+
+			// Combine base and filtered pre-multiplied blend curves
+			UE::Anim::FNamedValueArrayUtils::Union(
+				OutCurve,
+				TargetCurve,
+				[this](UE::Anim::FCurveElement& InOutBaseElement, const UE::Anim::FCurveElement& InBlendElement, UE::Anim::ENamedValueUnionFlags InFlags)
+					{
+						InOutBaseElement.Value = FMath::Lerp(InOutBaseElement.Value, InBlendElement.Value, BlendWeight);
+						InOutBaseElement.Flags |= InBlendElement.Flags; // Should this only apply when the weight is relevant?
+					});
+		}
+	}
+
+	// Blend attributes
+	{
+		using namespace UE::Anim;
+
+		FStackAttributeContainer& BaseAttributes = BasePoseContext.CustomAttributes;
+		FStackAttributeContainer& BlendAttributes = TargetAttributes;
+		FStackAttributeContainer& OutAttributes = OutputPoseData.GetAttributes();
+		OutAttributes.CopyFrom(BaseAttributes);
+
+		// Premultiply masked blend attributes
+		const TArray<TWeakObjectPtr<UScriptStruct>>& BlendUniqueTypes = BlendAttributes.GetUniqueTypes();
+
+		//CachedAttributeMaskWeights.ForEachElement([&BlendUniqueTypes, &BlendAttributes](const FNamedFloat& AttributeMask)
+		for (const FNamedFloat& AttributeMask : CachedAttributeMaskWeights)
+		{
+			for (int32 BlendUniqueTypeIndex = 0; BlendUniqueTypeIndex < BlendUniqueTypes.Num(); ++BlendUniqueTypeIndex)
+			{
+				const TArray<FAttributeId>& TypedAttributeIdentifiers = BlendAttributes.GetKeys(BlendUniqueTypeIndex);
+				const TArray<UE::Anim::TWrappedAttribute<FAnimStackAllocator>>& TypedAttributeValues = BlendAttributes.GetValues(BlendUniqueTypeIndex);
+
+				const TWeakObjectPtr<UScriptStruct> AttributeType = BlendUniqueTypes[BlendUniqueTypeIndex];
+				const IAttributeBlendOperator* Operator = UE::Anim::AttributeTypes::GetTypeOperator(AttributeType);
+
+				TWrappedAttribute<FAnimStackAllocator> DefaultData;
+				AttributeType->InitializeStruct(DefaultData.GetPtr<void>());
+
+				const int32 AttributeIndex = TypedAttributeIdentifiers.IndexOfByPredicate([&AttributeMask](const FAttributeId& InAttribute)
+					{
+						return InAttribute.GetName() == AttributeMask.Name;
+					});
+
+				if (AttributeIndex != INDEX_NONE)
+				{
+					const FAttributeId& AttributeId = TypedAttributeIdentifiers[AttributeIndex];
+
+					if (AttributeMask.Value == 0.0f)
+					{
+						BlendAttributes.Remove(AttributeType.Get(), AttributeId);
+					}
+					else
+					{
+						uint8* Value = BlendAttributes.Find(AttributeType.Get(), AttributeId);
+						check(Value != nullptr);
+
+						Operator->Interpolate(DefaultData.GetPtr<void>(), Value, AttributeMask.Value, Value);
+					}
+					break;
+				}
+			}
+		}
+
+		const TArray<const UE::Anim::FStackAttributeContainer> SourceAttributes { BaseAttributes, BlendAttributes };
+		const TArray<const float> SourceWeights { 1.0f - BlendWeight, BlendWeight };
+		static const TArray<const int32> WeightView { 0, 1 };
+
+		UE::Anim::Attributes::BlendAttributes(SourceAttributes, SourceWeights, WeightView, BaseAttributes);
+	}
+}
+
+void FAnimNode_BlendProfileLayeredBlend::GatherDebugData(FNodeDebugData& DebugData)
+{
+	DECLARE_SCOPE_HIERARCHICAL_COUNTER_ANIMNODE(GatherDebugData)
+
+	BasePose.GatherDebugData(DebugData.BranchFlow(1.f));
+	BlendPose.GatherDebugData(DebugData.BranchFlow(BlendWeight));
+}
+
+void FAnimNode_BlendProfileLayeredBlend::CreateMaskWeights(const USkeleton* Skeleton)
+{
+	if (!Skeleton)
+	{
+		return;
+	}
+
+	const FReferenceSkeleton& RefSkeleton = Skeleton->GetReferenceSkeleton();
+
+	const int32 NumBones = RefSkeleton.GetNum();
+	PerBoneBlendWeights.Reset(NumBones);
+	PerBoneBlendWeights.AddZeroed(NumBones);
+
+	check(MaskTable->IsTableType<FHierarchyTable_TableType_Skeleton>());
+	check(MaskTable->IsElementType<FHierarchyTable_ElementType_Mask>());
+
+	const FHierarchyTable_TableType_Skeleton& SkeletonMetadata = MaskTable->GetTableMetadata<FHierarchyTable_TableType_Skeleton>();
+
+	for (int32 EntryIndex = 0; EntryIndex < MaskTable->GetTableData().Num(); ++EntryIndex)
+	{
+		const FHierarchyTableEntryData* const EntryData = MaskTable->GetTableEntry(EntryIndex);
+		if (EntryData->TablePayload.Get<FHierarchyTable_TablePayloadType_Skeleton>().EntryType != ESkeletonHierarchyTable_TablePayloadEntryType::Bone)
+		{
+			break;
+		}
+
+		const FHierarchyTable_ElementType_Mask* MaskValue = EntryData->GetValue<FHierarchyTable_ElementType_Mask>();
+		check(MaskValue);
+			
+		PerBoneBlendWeights[EntryIndex] = MaskValue->Value;
+	}
+}
+
+void FAnimNode_BlendProfileLayeredBlend::UpdateDesiredBoneWeight()
+{
+	check(CurrentBoneBlendWeights.Num() == DesiredBoneBlendWeights.Num());
+
+	CurrentBoneBlendWeights.Init(0, CurrentBoneBlendWeights.Num());
+
+	for (int32 BoneIndex = 0; BoneIndex < DesiredBoneBlendWeights.Num(); ++BoneIndex)
+	{
+		float TargetBlendWeight = BlendWeight * DesiredBoneBlendWeights[BoneIndex];
+		
+		if (FAnimWeight::IsRelevant(TargetBlendWeight))
+		{
+			CurrentBoneBlendWeights[BoneIndex] = TargetBlendWeight;
+		}
+	}
+}
diff --git a/Engine/Plugins/Experimental/Animation/HierarchyTableAnimation/Source/Runtime/Private/HierachyTableAnimationRuntimeModule.cpp b/Engine/Plugins/Experimental/Animation/HierarchyTableAnimation/Source/Runtime/Private/HierachyTableAnimationRuntimeModule.cpp
new file mode 100644
--- /dev/null
+++ b/Engine/Plugins/Experimental/Animation/HierarchyTableAnimation/Source/Runtime/Private/HierachyTableAnimationRuntimeModule.cpp
@@ -0,0 +1,14 @@
+// Copyright Epic Games, Inc. All Rights Reserved.
+
+#include "HierarchyTableAnimationRuntimeModule.h"
+#include "Modules/ModuleManager.h"
+
+void FHierarchyTableAnimationRuntimeModule::StartupModule()
+{
+}
+
+void FHierarchyTableAnimationRuntimeModule::ShutdownModule()
+{
+}
+
+IMPLEMENT_MODULE(FHierarchyTableAnimationRuntimeModule, HierarchyTableAnimationRuntime)
diff --git a/Engine/Plugins/Experimental/Animation/HierarchyTableAnimation/Source/Runtime/Public/AnimNode_BlendProfileLayeredBlend.h b/Engine/Plugins/Experimental/Animation/HierarchyTableAnimation/Source/Runtime/Public/AnimNode_BlendProfileLayeredBlend.h
new file mode 100644
--- /dev/null
+++ b/Engine/Plugins/Experimental/Animation/HierarchyTableAnimation/Source/Runtime/Public/AnimNode_BlendProfileLayeredBlend.h
@@ -0,0 +1,117 @@
+// Copyright Epic Games, Inc. All Rights Reserved.
+
+#pragma once
+
+#include "CoreMinimal.h"
+#include "UObject/ObjectMacros.h"
+#include "Animation/AnimTypes.h"
+#include "Animation/AnimNodeBase.h"
+#include "AnimNode_BlendProfileLayeredBlend.generated.h"
+
+class UHierarchyTable;
+enum class EHierarchyTableEntryType : uint8;
+
+USTRUCT(BlueprintInternalUseOnly)
+struct HIERARCHYTABLEANIMATIONRUNTIME_API FAnimNode_BlendProfileLayeredBlend : public FAnimNode_Base
+{
+	GENERATED_BODY()
+
+public:
+	/** The source pose */
+	UPROPERTY(EditAnywhere, BlueprintReadWrite, Category = Links)
+	FPoseLink BasePose;
+
+	/** The target pose */
+	UPROPERTY(EditAnywhere, BlueprintReadWrite, Category = Links, meta = (BlueprintCompilerGeneratedDefaults))
+	FPoseLink BlendPose;
+
+	/**
+	 * The mask hierarchy table to use to control layering of the pose, curves, and attributes
+	 */
+	UPROPERTY(EditAnywhere, BlueprintReadWrite, Category = Config)
+	TObjectPtr<UHierarchyTable> MaskTable;
+
+protected:
+	// Per-bone weights for the skeleton. Serialized as these are only relative to the skeleton, but can potentially
+	// be regenerated at runtime if the GUIDs dont match
+	UPROPERTY()
+	TArray<float> PerBoneBlendWeights;
+
+	// transient data to handle weight and target weight
+	// this array changes based on required bones
+	TArray<float> DesiredBoneBlendWeights;
+	TArray<float> CurrentBoneBlendWeights;
+
+	UE::Anim::TNamedValueArray<FDefaultAllocator, UE::Anim::FCurveElement> CachedCurveMaskWeights;
+
+	struct FNamedFloat
+	{
+		FName Name;
+		float Value;
+
+		FNamedFloat(const FName InName, const float InValue)
+			: Name(InName)
+			, Value(InValue)
+		{
+		}
+	};
+	
+	TArray<FNamedFloat> CachedAttributeMaskWeights;
+
+	// Guids for skeleton used to determine whether the PerBoneBlendWeights need rebuilding
+	UPROPERTY()
+	FGuid SkeletonGuid;
+
+	// Guid for virtual bones used to determine whether the PerBoneBlendWeights need rebuilding
+	UPROPERTY()
+	FGuid VirtualBoneGuid;
+
+	/** The weight of target pose */
+	UPROPERTY(EditAnywhere, BlueprintReadWrite, Category = Runtime, meta = (BlueprintCompilerGeneratedDefaults, PinShownByDefault))
+	float BlendWeight;
+
+	// Serial number of the required bones container
+	uint16 RequiredBonesSerialNumber;
+
+	/** Whether to incorporate the per-bone blend weight of the root bone when lending root motion */
+	UPROPERTY(EditAnywhere, Category = Config)
+	bool bBlendRootMotionBasedOnRootBone;
+
+public:
+	FAnimNode_BlendProfileLayeredBlend()
+		: BlendWeight(1.0f)
+		, RequiredBonesSerialNumber(0)
+		, bBlendRootMotionBasedOnRootBone(true)
+	{
+	}
+
+	// FAnimNode_Base interface
+	virtual void Initialize_AnyThread(const FAnimationInitializeContext& Context) override;
+	virtual void CacheBones_AnyThread(const FAnimationCacheBonesContext& Context) override;
+	virtual void Update_AnyThread(const FAnimationUpdateContext& Context) override;
+	virtual void Evaluate_AnyThread(FPoseContext& Output) override;
+	virtual void GatherDebugData(FNodeDebugData& DebugData) override;
+	// End of FAnimNode_Base interface
+
+	// Invalidate the cached per-bone blend weights from the skeleton
+	void InvalidatePerBoneBlendWeights() { RequiredBonesSerialNumber = 0; SkeletonGuid = FGuid(); VirtualBoneGuid = FGuid(); }
+
+	// Invalidates the cached bone data so it is recalculated the next time this node is updated
+	void InvalidateCachedBoneData() { RequiredBonesSerialNumber = 0; }
+
+	void CreateMaskWeights(const USkeleton* Skeleton);
+
+private:
+	// Rebuild cache per bone blend weights from the skeleton
+	void RebuildPerBoneBlendWeights(const USkeleton* InSkeleton);
+
+	// Check whether per-bone blend weights are valid according to the skeleton (GUID check)
+	bool ArePerBoneBlendWeightsValid(const USkeleton* InSkeleton) const;
+
+	// Update cached data if required
+	void UpdateCachedBoneData(const FBoneContainer& RequiredBones, const USkeleton* Skeleton);
+
+	friend class UAnimGraphNode_BlendProfileLayeredBlend;
+
+	void UpdateDesiredBoneWeight();
+};
\ No newline at end of file
diff --git a/Engine/Plugins/Experimental/Animation/HierarchyTableAnimation/Source/Runtime/Public/HierarchyTableAnimationRuntimeModule.h b/Engine/Plugins/Experimental/Animation/HierarchyTableAnimation/Source/Runtime/Public/HierarchyTableAnimationRuntimeModule.h
new file mode 100644
--- /dev/null
+++ b/Engine/Plugins/Experimental/Animation/HierarchyTableAnimation/Source/Runtime/Public/HierarchyTableAnimationRuntimeModule.h
@@ -0,0 +1,13 @@
+// Copyright Epic Games, Inc. All Rights Reserved.
+
+#pragma once
+
+#include "Modules/ModuleInterface.h"
+
+class FHierarchyTableAnimationRuntimeModule : public IModuleInterface
+{
+public:
+	virtual void StartupModule() override;
+	
+	virtual void ShutdownModule() override;
+};
\ No newline at end of file
diff --git a/Engine/Plugins/Experimental/Animation/HierarchyTableAnimation/Source/Runtime/Public/MaskProfile/HierarchyTableTypeMask.h b/Engine/Plugins/Experimental/Animation/HierarchyTableAnimation/Source/Runtime/Public/MaskProfile/HierarchyTableTypeMask.h
new file mode 100644
--- /dev/null
+++ b/Engine/Plugins/Experimental/Animation/HierarchyTableAnimation/Source/Runtime/Public/MaskProfile/HierarchyTableTypeMask.h
@@ -0,0 +1,21 @@
+// Copyright Epic Games, Inc. All Rights Reserved.
+
+#pragma once
+
+#include "HierarchyTableType.h"
+
+#include "HierarchyTableTypeMask.generated.h"
+
+USTRUCT()
+struct FHierarchyTable_ElementType_Mask final : public FHierarchyTable_ElementType
+{
+	GENERATED_BODY()
+
+	FHierarchyTable_ElementType_Mask()
+		: Value(1.0f)
+	{
+	}
+
+	UPROPERTY()
+	float Value;
+};
\ No newline at end of file
diff --git a/Engine/Plugins/Experimental/Animation/HierarchyTableAnimation/Source/Runtime/Public/SkeletonHierarchyTableType.h b/Engine/Plugins/Experimental/Animation/HierarchyTableAnimation/Source/Runtime/Public/SkeletonHierarchyTableType.h
new file mode 100644
--- /dev/null
+++ b/Engine/Plugins/Experimental/Animation/HierarchyTableAnimation/Source/Runtime/Public/SkeletonHierarchyTableType.h
@@ -0,0 +1,35 @@
+// Copyright Epic Games, Inc. All Rights Reserved.
+
+#pragma once
+
+#include "HierarchyTableType.h"
+
+#include "SkeletonHierarchyTableType.generated.h"
+
+class USkeleton;
+
+USTRUCT()
+struct FHierarchyTable_TableType_Skeleton final : public FHierarchyTable_TableType
+{
+	GENERATED_BODY()
+
+	UPROPERTY()
+	TObjectPtr<USkeleton> Skeleton;
+};
+
+UENUM()
+enum class ESkeletonHierarchyTable_TablePayloadEntryType
+{
+	Bone,
+	Curve,
+	Attribute
+};
+
+USTRUCT()
+struct FHierarchyTable_TablePayloadType_Skeleton final : public FHierarchyTable_TablePayloadType
+{
+	GENERATED_BODY()
+
+	UPROPERTY()
+	ESkeletonHierarchyTable_TablePayloadEntryType EntryType;
+};
\ No newline at end of file
diff --git a/Engine/Plugins/Experimental/Animation/HierarchyTableAnimation/Source/Runtime/Public/TimeProfile/HierarchyTableTypeTime.h b/Engine/Plugins/Experimental/Animation/HierarchyTableAnimation/Source/Runtime/Public/TimeProfile/HierarchyTableTypeTime.h
new file mode 100644
--- /dev/null
+++ b/Engine/Plugins/Experimental/Animation/HierarchyTableAnimation/Source/Runtime/Public/TimeProfile/HierarchyTableTypeTime.h
@@ -0,0 +1,29 @@
+// Copyright Epic Games, Inc. All Rights Reserved.
+
+#pragma once
+
+#include "HierarchyTableType.h"
+
+#include "HierarchyTableTypeTime.generated.h"
+
+USTRUCT()
+struct FHierarchyTable_ElementType_Time final : public FHierarchyTable_ElementType
+{
+	GENERATED_BODY()
+
+	FHierarchyTable_ElementType_Time()
+		: StartTime(0.0f)
+		, EndTime(0.0f)
+		, TimeFactor(0.0f)
+	{
+	}
+
+	UPROPERTY()
+	float StartTime;
+
+	UPROPERTY()
+	float EndTime;
+
+	UPROPERTY()
+	float TimeFactor;
+};
\ No newline at end of file
diff --git a/Engine/Plugins/Experimental/Animation/HierarchyTableAnimation/Source/UncookedOnly/HierarchyTableAnimationUncookedOnly.Build.cs b/Engine/Plugins/Experimental/Animation/HierarchyTableAnimation/Source/UncookedOnly/HierarchyTableAnimationUncookedOnly.Build.cs
new file mode 100644
--- /dev/null
+++ b/Engine/Plugins/Experimental/Animation/HierarchyTableAnimation/Source/UncookedOnly/HierarchyTableAnimationUncookedOnly.Build.cs
@@ -0,0 +1,46 @@
+// Copyright Epic Games, Inc. All Rights Reserved.
+
+using EpicGames.Core;
+using UnrealBuildTool;
+
+public class HierarchyTableAnimationUncookedOnly : ModuleRules
+{
+	public HierarchyTableAnimationUncookedOnly(ReadOnlyTargetRules Target) : base(Target)
+	{
+		PublicDependencyModuleNames.AddRange(
+			new string[]
+			{
+				"AnimationCore",
+				"AnimGraph",
+				"AnimGraphRuntime",
+				"AnimationBlueprintLibrary",
+				"ToolMenus"
+			});
+
+		PrivateDependencyModuleNames.AddAll(
+			"Core",
+			"CoreUObject",
+			"Engine",
+			"InputCore",
+			"Slate",
+			"SlateCore",
+			"UnrealEd",
+			"HierarchyTableRuntime",
+			"HierarchyTableEditor",
+			"HierarchyTableAnimationRuntime"
+		);
+
+		if (Target.bBuildEditor == true)
+		{
+			PrivateDependencyModuleNames.AddRange(
+				new string[]
+				{
+						"BlueprintGraph",
+						"EditorFramework",
+						"Kismet",
+						"UnrealEd",
+				}
+			);
+		}
+	}
+}
diff --git a/Engine/Plugins/Experimental/Animation/HierarchyTableAnimation/Source/UncookedOnly/Private/AnimGraphNode_BlendProfileLayeredBlend.cpp b/Engine/Plugins/Experimental/Animation/HierarchyTableAnimation/Source/UncookedOnly/Private/AnimGraphNode_BlendProfileLayeredBlend.cpp
new file mode 100644
--- /dev/null
+++ b/Engine/Plugins/Experimental/Animation/HierarchyTableAnimation/Source/UncookedOnly/Private/AnimGraphNode_BlendProfileLayeredBlend.cpp
@@ -0,0 +1,84 @@
+// Copyright Epic Games, Inc. All Rights Reserved.
+
+#include "AnimGraphNode_BlendProfileLayeredBlend.h"
+#include "ToolMenus.h"
+#include "Kismet2/BlueprintEditorUtils.h"
+
+#include "AnimGraphCommands.h"
+#include "ScopedTransaction.h"
+
+#include "DetailLayoutBuilder.h"
+#include "Kismet2/CompilerResultsLog.h"
+
+#include "HierarchyTable.h"
+#include "MaskProfile/HierarchyTableTypeMask.h"
+
+#define LOCTEXT_NAMESPACE "AnimGraphNode_BlendProfileLayeredBlend"
+
+UAnimGraphNode_BlendProfileLayeredBlend::UAnimGraphNode_BlendProfileLayeredBlend(const FObjectInitializer& ObjectInitializer)
+	: Super(ObjectInitializer)
+{
+}
+
+FLinearColor UAnimGraphNode_BlendProfileLayeredBlend::GetNodeTitleColor() const
+{
+	return FLinearColor(0.75f, 0.75f, 0.75f);
+}
+
+FText UAnimGraphNode_BlendProfileLayeredBlend::GetTooltipText() const
+{
+	return LOCTEXT("AnimGraphNode_BlendProfileLayeredBlend_Tooltip", "Profile Blend");
+}
+
+FText UAnimGraphNode_BlendProfileLayeredBlend::GetNodeTitle(ENodeTitleType::Type TitleType) const
+{
+	return LOCTEXT("AnimGraphNode_BlendProfileLayeredBlend_Title", "Profile Blend");
+}
+
+FString UAnimGraphNode_BlendProfileLayeredBlend::GetNodeCategory() const
+{
+	return TEXT("Animation|Blends");
+}
+
+void UAnimGraphNode_BlendProfileLayeredBlend::PreloadRequiredAssets()
+{
+	TObjectPtr<UHierarchyTable> MaskTable = Node.MaskTable;
+	if (MaskTable)
+	{
+		PreloadObject(MaskTable);
+	}
+
+	Super::PreloadRequiredAssets();
+}
+
+void UAnimGraphNode_BlendProfileLayeredBlend::ValidateAnimNodeDuringCompilation(class USkeleton* ForSkeleton, class FCompilerResultsLog& MessageLog)
+{
+	UAnimGraphNode_Base::ValidateAnimNodeDuringCompilation(ForSkeleton, MessageLog);
+
+	bool bCompilationError = false;
+
+	const TObjectPtr<UHierarchyTable> MaskTable = Node.MaskTable;
+
+	if (MaskTable == nullptr)
+	{
+		MessageLog.Error(*LOCTEXT("InvalidMask", "@@ has a null mask table.").ToString(), this, MaskTable);
+		bCompilationError = true;
+	}
+	else if (MaskTable->IsElementType<FHierarchyTable_ElementType_Mask>())
+	{
+		MessageLog.Error(*LOCTEXT("InvalidMaskType", "@@ has a null mask table that is not a mask type.").ToString(), this, MaskTable);
+		bCompilationError = true;
+	}
+
+	if (bCompilationError)
+	{
+		return;
+	}
+
+	if (!Node.ArePerBoneBlendWeightsValid(ForSkeleton))
+	{
+		Node.RebuildPerBoneBlendWeights(ForSkeleton);
+	}
+}
+
+#undef LOCTEXT_NAMESPACE
diff --git a/Engine/Plugins/Experimental/Animation/HierarchyTableAnimation/Source/UncookedOnly/Private/HierachyTableAnimationUncookedOnlyModule.cpp b/Engine/Plugins/Experimental/Animation/HierarchyTableAnimation/Source/UncookedOnly/Private/HierachyTableAnimationUncookedOnlyModule.cpp
new file mode 100644
--- /dev/null
+++ b/Engine/Plugins/Experimental/Animation/HierarchyTableAnimation/Source/UncookedOnly/Private/HierachyTableAnimationUncookedOnlyModule.cpp
@@ -0,0 +1,14 @@
+// Copyright Epic Games, Inc. All Rights Reserved.
+
+#include "HierarchyTableAnimationUncookedOnlyModule.h"
+#include "Modules/ModuleManager.h"
+
+void FHierarchyTableAnimationUncookedOnlyModule::StartupModule()
+{
+}
+
+void FHierarchyTableAnimationUncookedOnlyModule::ShutdownModule()
+{
+}
+
+IMPLEMENT_MODULE(FHierarchyTableAnimationUncookedOnlyModule, HierarchyTableAnimationUncookedOnly)
diff --git a/Engine/Plugins/Experimental/Animation/HierarchyTableAnimation/Source/UncookedOnly/Public/AnimGraphNode_BlendProfileLayeredBlend.h b/Engine/Plugins/Experimental/Animation/HierarchyTableAnimation/Source/UncookedOnly/Public/AnimGraphNode_BlendProfileLayeredBlend.h
new file mode 100644
--- /dev/null
+++ b/Engine/Plugins/Experimental/Animation/HierarchyTableAnimation/Source/UncookedOnly/Public/AnimGraphNode_BlendProfileLayeredBlend.h
@@ -0,0 +1,30 @@
+// Copyright Epic Games, Inc. All Rights Reserved.
+
+#pragma once
+
+#include "CoreMinimal.h"
+#include "UObject/ObjectMacros.h"
+#include "AnimGraphNode_BlendListBase.h"
+#include "AnimNode_BlendProfileLayeredBlend.h"
+#include "AnimGraphNode_BlendProfileLayeredBlend.generated.h"
+
+UCLASS(MinimalAPI)
+class UAnimGraphNode_BlendProfileLayeredBlend : public UAnimGraphNode_BlendListBase
+{
+	GENERATED_UCLASS_BODY()
+
+	UPROPERTY(EditAnywhere, Category=Settings)
+	FAnimNode_BlendProfileLayeredBlend Node;
+
+	// UEdGraphNode interface
+	virtual FLinearColor GetNodeTitleColor() const override;
+	virtual FText GetTooltipText() const override;
+	virtual FText GetNodeTitle(ENodeTitleType::Type TitleType) const override;
+	// End of UEdGraphNode interface
+
+	// UAnimGraphNode_Base interface
+	virtual FString GetNodeCategory() const override;
+	virtual void PreloadRequiredAssets() override;
+	virtual void ValidateAnimNodeDuringCompilation(class USkeleton* ForSkeleton, class FCompilerResultsLog& MessageLog) override;
+	// End of UAnimGraphNode_Base interface
+};
diff --git a/Engine/Plugins/Experimental/Animation/HierarchyTableAnimation/Source/UncookedOnly/Public/HierarchyTableAnimationUncookedOnlyModule.h b/Engine/Plugins/Experimental/Animation/HierarchyTableAnimation/Source/UncookedOnly/Public/HierarchyTableAnimationUncookedOnlyModule.h
new file mode 100644
--- /dev/null
+++ b/Engine/Plugins/Experimental/Animation/HierarchyTableAnimation/Source/UncookedOnly/Public/HierarchyTableAnimationUncookedOnlyModule.h
@@ -0,0 +1,13 @@
+// Copyright Epic Games, Inc. All Rights Reserved.
+
+#pragma once
+
+#include "Modules/ModuleInterface.h"
+
+class FHierarchyTableAnimationUncookedOnlyModule : public IModuleInterface
+{
+public:
+	virtual void StartupModule() override;
+	
+	virtual void ShutdownModule() override;
+};
\ No newline at end of file
diff --git a/Engine/Plugins/Experimental/Animation/HierarchyTableBuiltin/HierarchyTableBuiltin.uplugin b/Engine/Plugins/Experimental/Animation/HierarchyTableBuiltin/HierarchyTableBuiltin.uplugin
deleted file mode 100644
diff --git a/Engine/Plugins/Experimental/Animation/HierarchyTableBuiltin/Source/Editor/HierarchyTableBuiltinEditor.Build.cs b/Engine/Plugins/Experimental/Animation/HierarchyTableBuiltin/Source/Editor/HierarchyTableBuiltinEditor.Build.cs
deleted file mode 100644
diff --git a/Engine/Plugins/Experimental/Animation/HierarchyTableBuiltin/Source/Editor/Private/HierachyTableBuiltinEditorModule.cpp b/Engine/Plugins/Experimental/Animation/HierarchyTableBuiltin/Source/Editor/Private/HierachyTableBuiltinEditorModule.cpp
deleted file mode 100644
diff --git a/Engine/Plugins/Experimental/Animation/HierarchyTableBuiltin/Source/Editor/Private/MaskProfile/MaskProfileFactory.cpp b/Engine/Plugins/Experimental/Animation/HierarchyTableBuiltin/Source/Editor/Private/MaskProfile/MaskProfileFactory.cpp
deleted file mode 100644
diff --git a/Engine/Plugins/Experimental/Animation/HierarchyTableBuiltin/Source/Editor/Private/MaskProfile/MaskProfileFactory.h b/Engine/Plugins/Experimental/Animation/HierarchyTableBuiltin/Source/Editor/Private/MaskProfile/MaskProfileFactory.h
deleted file mode 100644
diff --git a/Engine/Plugins/Experimental/Animation/HierarchyTableBuiltin/Source/Editor/Private/MaskProfile/MaskProfileWidgetConstructor.cpp b/Engine/Plugins/Experimental/Animation/HierarchyTableBuiltin/Source/Editor/Private/MaskProfile/MaskProfileWidgetConstructor.cpp
deleted file mode 100644
diff --git a/Engine/Plugins/Experimental/Animation/HierarchyTableBuiltin/Source/Editor/Private/MaskProfile/MaskProfileWidgetConstructor.h b/Engine/Plugins/Experimental/Animation/HierarchyTableBuiltin/Source/Editor/Private/MaskProfile/MaskProfileWidgetConstructor.h
deleted file mode 100644
diff --git a/Engine/Plugins/Experimental/Animation/HierarchyTableBuiltin/Source/Editor/Private/SkeletonHierarchyTableTypeHandler.cpp b/Engine/Plugins/Experimental/Animation/HierarchyTableBuiltin/Source/Editor/Private/SkeletonHierarchyTableTypeHandler.cpp
deleted file mode 100644
diff --git a/Engine/Plugins/Experimental/Animation/HierarchyTableBuiltin/Source/Editor/Private/SkeletonHierarchyTableTypeHandler.h b/Engine/Plugins/Experimental/Animation/HierarchyTableBuiltin/Source/Editor/Private/SkeletonHierarchyTableTypeHandler.h
deleted file mode 100644
diff --git a/Engine/Plugins/Experimental/Animation/HierarchyTableBuiltin/Source/Editor/Private/TimeProfile/TimeProfileFactory.cpp b/Engine/Plugins/Experimental/Animation/HierarchyTableBuiltin/Source/Editor/Private/TimeProfile/TimeProfileFactory.cpp
deleted file mode 100644
diff --git a/Engine/Plugins/Experimental/Animation/HierarchyTableBuiltin/Source/Editor/Private/TimeProfile/TimeProfileFactory.h b/Engine/Plugins/Experimental/Animation/HierarchyTableBuiltin/Source/Editor/Private/TimeProfile/TimeProfileFactory.h
deleted file mode 100644
diff --git a/Engine/Plugins/Experimental/Animation/HierarchyTableBuiltin/Source/Editor/Private/TimeProfile/TimeProfileWidgetConstructor.cpp b/Engine/Plugins/Experimental/Animation/HierarchyTableBuiltin/Source/Editor/Private/TimeProfile/TimeProfileWidgetConstructor.cpp
deleted file mode 100644
diff --git a/Engine/Plugins/Experimental/Animation/HierarchyTableBuiltin/Source/Editor/Private/TimeProfile/TimeProfileWidgetConstructor.h b/Engine/Plugins/Experimental/Animation/HierarchyTableBuiltin/Source/Editor/Private/TimeProfile/TimeProfileWidgetConstructor.h
deleted file mode 100644
diff --git a/Engine/Plugins/Experimental/Animation/HierarchyTableBuiltin/Source/Editor/Public/HierarchyTableBuiltinEditorModule.h b/Engine/Plugins/Experimental/Animation/HierarchyTableBuiltin/Source/Editor/Public/HierarchyTableBuiltinEditorModule.h
deleted file mode 100644
diff --git a/Engine/Plugins/Experimental/Animation/HierarchyTableBuiltin/Source/Editor/Public/MaskProfile/MaskProfileProxyColumn.h b/Engine/Plugins/Experimental/Animation/HierarchyTableBuiltin/Source/Editor/Public/MaskProfile/MaskProfileProxyColumn.h
deleted file mode 100644
diff --git a/Engine/Plugins/Experimental/Animation/HierarchyTableBuiltin/Source/Editor/Public/TimeProfile/TimeProfileProxyColumn.h b/Engine/Plugins/Experimental/Animation/HierarchyTableBuiltin/Source/Editor/Public/TimeProfile/TimeProfileProxyColumn.h
deleted file mode 100644
diff --git a/Engine/Plugins/Experimental/Animation/HierarchyTableBuiltin/Source/Runtime/HierarchyTableBuiltinRuntime.Build.cs b/Engine/Plugins/Experimental/Animation/HierarchyTableBuiltin/Source/Runtime/HierarchyTableBuiltinRuntime.Build.cs
deleted file mode 100644
diff --git a/Engine/Plugins/Experimental/Animation/HierarchyTableBuiltin/Source/Runtime/Private/AnimNode_BlendProfileLayeredBlend.cpp b/Engine/Plugins/Experimental/Animation/HierarchyTableBuiltin/Source/Runtime/Private/AnimNode_BlendProfileLayeredBlend.cpp
deleted file mode 100644
diff --git a/Engine/Plugins/Experimental/Animation/HierarchyTableBuiltin/Source/Runtime/Private/HierachyTableBuiltinRuntimeModule.cpp b/Engine/Plugins/Experimental/Animation/HierarchyTableBuiltin/Source/Runtime/Private/HierachyTableBuiltinRuntimeModule.cpp
deleted file mode 100644
diff --git a/Engine/Plugins/Experimental/Animation/HierarchyTableBuiltin/Source/Runtime/Public/AnimNode_BlendProfileLayeredBlend.h b/Engine/Plugins/Experimental/Animation/HierarchyTableBuiltin/Source/Runtime/Public/AnimNode_BlendProfileLayeredBlend.h
deleted file mode 100644
diff --git a/Engine/Plugins/Experimental/Animation/HierarchyTableBuiltin/Source/Runtime/Public/HierarchyTableBuiltinRuntimeModule.h b/Engine/Plugins/Experimental/Animation/HierarchyTableBuiltin/Source/Runtime/Public/HierarchyTableBuiltinRuntimeModule.h
deleted file mode 100644
diff --git a/Engine/Plugins/Experimental/Animation/HierarchyTableBuiltin/Source/Runtime/Public/MaskProfile/HierarchyTableTypeMask.h b/Engine/Plugins/Experimental/Animation/HierarchyTableBuiltin/Source/Runtime/Public/MaskProfile/HierarchyTableTypeMask.h
deleted file mode 100644
diff --git a/Engine/Plugins/Experimental/Animation/HierarchyTableBuiltin/Source/Runtime/Public/SkeletonHierarchyTableType.h b/Engine/Plugins/Experimental/Animation/HierarchyTableBuiltin/Source/Runtime/Public/SkeletonHierarchyTableType.h
deleted file mode 100644
diff --git a/Engine/Plugins/Experimental/Animation/HierarchyTableBuiltin/Source/Runtime/Public/TimeProfile/HierarchyTableTypeTime.h b/Engine/Plugins/Experimental/Animation/HierarchyTableBuiltin/Source/Runtime/Public/TimeProfile/HierarchyTableTypeTime.h
deleted file mode 100644
diff --git a/Engine/Plugins/Experimental/Animation/HierarchyTableBuiltin/Source/UncookedOnly/HierarchyTableBuiltinUncookedOnly.Build.cs b/Engine/Plugins/Experimental/Animation/HierarchyTableBuiltin/Source/UncookedOnly/HierarchyTableBuiltinUncookedOnly.Build.cs
deleted file mode 100644
diff --git a/Engine/Plugins/Experimental/Animation/HierarchyTableBuiltin/Source/UncookedOnly/Private/AnimGraphNode_BlendProfileLayeredBlend.cpp b/Engine/Plugins/Experimental/Animation/HierarchyTableBuiltin/Source/UncookedOnly/Private/AnimGraphNode_BlendProfileLayeredBlend.cpp
deleted file mode 100644
diff --git a/Engine/Plugins/Experimental/Animation/HierarchyTableBuiltin/Source/UncookedOnly/Private/HierachyTableBuiltinUncookedOnlyModule.cpp b/Engine/Plugins/Experimental/Animation/HierarchyTableBuiltin/Source/UncookedOnly/Private/HierachyTableBuiltinUncookedOnlyModule.cpp
deleted file mode 100644
diff --git a/Engine/Plugins/Experimental/Animation/HierarchyTableBuiltin/Source/UncookedOnly/Public/AnimGraphNode_BlendProfileLayeredBlend.h b/Engine/Plugins/Experimental/Animation/HierarchyTableBuiltin/Source/UncookedOnly/Public/AnimGraphNode_BlendProfileLayeredBlend.h
deleted file mode 100644
diff --git a/Engine/Plugins/Experimental/Animation/HierarchyTableBuiltin/Source/UncookedOnly/Public/HierarchyTableBuiltinUncookedOnlyModule.h b/Engine/Plugins/Experimental/Animation/HierarchyTableBuiltin/Source/UncookedOnly/Public/HierarchyTableBuiltinUncookedOnlyModule.h
deleted file mode 100644
