Fixes to Hierarchy Table and API improvements
* Remove TEDS from HierarchyTables to simplify API
* Fix being able to add duplicate entries to table
* Fix table not updating when adding new values

diff --git a/Engine/Plugins/Experimental/Animation/HierarchyTable/HierarchyTable.uplugin b/Engine/Plugins/Experimental/Animation/HierarchyTable/HierarchyTable.uplugin
--- a/Engine/Plugins/Experimental/Animation/HierarchyTable/HierarchyTable.uplugin
+++ b/Engine/Plugins/Experimental/Animation/HierarchyTable/HierarchyTable.uplugin
@@ -28,9 +28,5 @@
 		}
 	],
 	"Plugins": [
-		{
-			"Name": "EditorDataStorageFeatures",
-			"Enabled": true
-		}
 	]
 }
diff --git a/Engine/Plugins/Experimental/Animation/HierarchyTable/Source/Editor/HierarchyTableEditor.Build.cs b/Engine/Plugins/Experimental/Animation/HierarchyTable/Source/Editor/HierarchyTableEditor.Build.cs
--- a/Engine/Plugins/Experimental/Animation/HierarchyTable/Source/Editor/HierarchyTableEditor.Build.cs
+++ b/Engine/Plugins/Experimental/Animation/HierarchyTable/Source/Editor/HierarchyTableEditor.Build.cs
@@ -17,12 +17,8 @@ public class HierarchyTableEditor : ModuleRules
 			"SlateCore",
 			"UnrealEd",
 			"HierarchyTableRuntime",
-			"TedsOutliner",
-			"TypedElementFramework",
-			"SceneOutliner",
 			"ToolMenus",
-			"ToolWidgets",
-			"TedsTableViewer"
+			"ToolWidgets"
 		);
 	}
 }
diff --git a/Engine/Plugins/Experimental/Animation/HierarchyTable/Source/Editor/Private/HierachyTableEditorModule.cpp b/Engine/Plugins/Experimental/Animation/HierarchyTable/Source/Editor/Private/HierachyTableEditorModule.cpp
--- a/Engine/Plugins/Experimental/Animation/HierarchyTable/Source/Editor/Private/HierachyTableEditorModule.cpp
+++ b/Engine/Plugins/Experimental/Animation/HierarchyTable/Source/Editor/Private/HierachyTableEditorModule.cpp
@@ -48,7 +48,7 @@ const TObjectPtr<UHierarchyTable_TableTypeHandler> FHierarchyTableEditorModule::
 	return nullptr;
 }
 
-void FHierarchyTableEditorModule::RegisterElementTypeEditorColumns(const UScriptStruct* ElementType, const TArray<TObjectPtr<UScriptStruct>> Columns)
+void FHierarchyTableEditorModule::RegisterElementTypeEditorColumns(const UScriptStruct* ElementType, const TArray<TSharedPtr<IHierarchyTableColumn>> Columns)
 {
 	EditorColumns.Add(ElementType, Columns);
 }
@@ -58,15 +58,16 @@ void FHierarchyTableEditorModule::UnregisterElementTypeEditorColumns(const UScri
 	EditorColumns.Remove(ElementType);
 }
 
-TArray<UScriptStruct*> FHierarchyTableEditorModule::GetElementTypeEditorColumns(const TObjectPtr<const UHierarchyTable> HierarchyTable)
+TArray<TSharedPtr<IHierarchyTableColumn>> FHierarchyTableEditorModule::GetElementTypeEditorColumns(const TObjectPtr<const UHierarchyTable> HierarchyTable)
 {
-	const TArray<UScriptStruct*>* Result = EditorColumns.Find(HierarchyTable->GetElementType());
+	const TArray<TSharedPtr<IHierarchyTableColumn>>* Result = EditorColumns.Find(HierarchyTable->GetElementType());
 
 	if (Result)
 	{
 		return *Result;
 	}
-	return TArray<UScriptStruct*>();
+
+	return TArray<TSharedPtr<IHierarchyTableColumn>>();
 }
 
 IMPLEMENT_MODULE(FHierarchyTableEditorModule, HierarchyTableEditor)
diff --git a/Engine/Plugins/Experimental/Animation/HierarchyTable/Source/Editor/Private/HierarchyTable/Factories/OverrideFactory.cpp b/Engine/Plugins/Experimental/Animation/HierarchyTable/Source/Editor/Private/HierarchyTable/Factories/OverrideFactory.cpp
deleted file mode 100644
diff --git a/Engine/Plugins/Experimental/Animation/HierarchyTable/Source/Editor/Private/HierarchyTable/Factories/OverrideFactory.h b/Engine/Plugins/Experimental/Animation/HierarchyTable/Source/Editor/Private/HierarchyTable/Factories/OverrideFactory.h
deleted file mode 100644
diff --git a/Engine/Plugins/Experimental/Animation/HierarchyTable/Source/Editor/Private/HierarchyTable/Widgets/OverrideWidgetConstructor.cpp b/Engine/Plugins/Experimental/Animation/HierarchyTable/Source/Editor/Private/HierarchyTable/Widgets/OverrideWidgetConstructor.cpp
deleted file mode 100644
diff --git a/Engine/Plugins/Experimental/Animation/HierarchyTable/Source/Editor/Private/HierarchyTable/Widgets/OverrideWidgetConstructor.h b/Engine/Plugins/Experimental/Animation/HierarchyTable/Source/Editor/Private/HierarchyTable/Widgets/OverrideWidgetConstructor.h
deleted file mode 100644
diff --git a/Engine/Plugins/Experimental/Animation/HierarchyTable/Source/Editor/Private/HierarchyTableEditorToolkit.cpp b/Engine/Plugins/Experimental/Animation/HierarchyTable/Source/Editor/Private/HierarchyTableEditorToolkit.cpp
--- a/Engine/Plugins/Experimental/Animation/HierarchyTable/Source/Editor/Private/HierarchyTableEditorToolkit.cpp
+++ b/Engine/Plugins/Experimental/Animation/HierarchyTable/Source/Editor/Private/HierarchyTableEditorToolkit.cpp
@@ -1,35 +1,20 @@
 // Copyright Epic Games, Inc. All Rights Reserved.
 
 #include "HierarchyTableEditorToolkit.h"
-#include "Widgets/Docking/SDockTab.h"
+
+#include "HierarchyTableEditorModule.h"
+#include "HierarchyTableType.h"
 #include "Modules/ModuleManager.h"
 #include "PropertyEditorModule.h"
-#include "TedsOutlinerModule.h"
-#include "Elements/Columns/TypedElementLabelColumns.h"
-#include "Elements/Columns/TypedElementTypeInfoColumns.h"
-#include "Elements/Columns/TypedElementSelectionColumns.h"
-#include "Elements/Columns/TypedElementMiscColumns.h"
-#include "Elements/Common/EditorDataStorageFeatures.h"
-#include "Elements/Framework/TypedElementQueryBuilder.h"
-#include "Elements/Interfaces/TypedElementDataStorageInterface.h"
-#include "SceneOutlinerPublicTypes.h"
-#include "TedsOutlinerMode.h"
-#include "ReferenceSkeleton.h"
-#include "Animation/Skeleton.h"
-#include "HierarchyTable/Columns/OverrideColumn.h"
-#include "HierarchyTableType.h"
+#include "SHierarchyTable.h"
 #include "ToolMenus.h"
-#include "Framework/Application/SlateApplication.h"
-#include "HierarchyTableEditorModule.h"
-#include "Widgets/Input/STextEntryPopup.h"
-#include "PersonaModule.h"
-#include "TedsOutlinerItem.h"
+#include "Widgets/Docking/SDockTab.h"
 
 #define LOCTEXT_NAMESPACE "HierarchyTableEditorToolkit"
 
 void FHierarchyTableEditorToolkit::InitEditor(const TArray<UObject*>& InObjects)
 {
-	HierarchyTable = Cast<UHierarchyTable>(InObjects[0]);
+	HierarchyTable = CastChecked<UHierarchyTable>(InObjects[0]);
 
 	const TSharedRef<FTabManager::FLayout> Layout = FTabManager::NewLayout("HierarchyTableEditorLayout")
 		->AddArea
@@ -60,44 +45,36 @@ void FHierarchyTableEditorToolkit::RegisterTabSpawners(const TSharedRef<class FT
 
 	WorkspaceMenuCategory = InTabManager->AddLocalWorkspaceMenuCategory(LOCTEXT("HierarchyTableEditor", "Hierarchy Table Editor"));
 
-	InTabManager->RegisterTabSpawner("HierarchyTableEditorTableTab", FOnSpawnTab::CreateLambda([this](const FSpawnTabArgs&)
-		{
-			return SNew(SDockTab)
-				[
-					CreateTedsOutliner()
-				];
-		}))
-		.SetDisplayName(LOCTEXT("HierarchyTable", "Hierarchy Table"))
-		.SetGroup(WorkspaceMenuCategory.ToSharedRef());
-
-
-	FPropertyEditorModule& PropertyEditorModule = FModuleManager::GetModuleChecked<FPropertyEditorModule>("PropertyEditor");
-	FDetailsViewArgs DetailsViewArgs;
-	DetailsViewArgs.NameAreaSettings = FDetailsViewArgs::HideNameArea;
-	TSharedRef<IDetailsView> DetailsView = PropertyEditorModule.CreateDetailView(DetailsViewArgs);
-	DetailsView->SetObjects(TArray<UObject*>{ HierarchyTable });
-	InTabManager->RegisterTabSpawner("HierarchyTableEditorDetailsTab", FOnSpawnTab::CreateLambda([=](const FSpawnTabArgs&)
-		{
-			return SNew(SDockTab)
-				[
-					DetailsView
-				];
-		}))
-		.SetDisplayName(INVTEXT("Details"))
-		.SetGroup(WorkspaceMenuCategory.ToSharedRef());
-}
-
-void FHierarchyTableEditorToolkit::OnClose()
-{
-	using namespace UE::Editor::DataStorage;
-	IEditorDataStorageProvider* DSI = GetMutableDataStorageFeature<IEditorDataStorageProvider>(StorageFeatureName);
+	// Table view
+	{
+		InTabManager->RegisterTabSpawner("HierarchyTableEditorTableTab", FOnSpawnTab::CreateLambda([this](const FSpawnTabArgs&)
+			{
+				return SNew(SDockTab)
+					[
+						SAssignNew(HierarchyTableWidget, SHierarchyTable, HierarchyTable)
+					];
+			}))
+			.SetDisplayName(LOCTEXT("HierarchyTable", "Hierarchy Table"))
+			.SetGroup(WorkspaceMenuCategory.ToSharedRef());
+	}
 
-	for (const TTuple<int32, RowHandle>& Row : EntryIndexToHandleMap)
+	// Details panel
 	{
-		DSI->RemoveRow(Row.Value);
+		FPropertyEditorModule& PropertyEditorModule = FModuleManager::GetModuleChecked<FPropertyEditorModule>("PropertyEditor");
+		FDetailsViewArgs DetailsViewArgs;
+		DetailsViewArgs.NameAreaSettings = FDetailsViewArgs::HideNameArea;
+		TSharedRef<IDetailsView> DetailsView = PropertyEditorModule.CreateDetailView(DetailsViewArgs);
+		DetailsView->SetObjects(TArray<UObject*>{ HierarchyTable });
+		InTabManager->RegisterTabSpawner("HierarchyTableEditorDetailsTab", FOnSpawnTab::CreateLambda([=](const FSpawnTabArgs&)
+			{
+				return SNew(SDockTab)
+					[
+						DetailsView
+					];
+			}))
+			.SetDisplayName(INVTEXT("Details"))
+			.SetGroup(WorkspaceMenuCategory.ToSharedRef());
 	}
-	
-	HierarchyTable->UnregisterOnEntryAdded(OnEntryAddedHandle);
 }
 
 void FHierarchyTableEditorToolkit::UnregisterTabSpawners(const TSharedRef<class FTabManager>& InTabManager)
@@ -107,95 +84,6 @@ void FHierarchyTableEditorToolkit::UnregisterTabSpawners(const TSharedRef<class
 	InTabManager->UnregisterTabSpawner("HierarchyTableEditorDetailsTab");
 }
 
-TSharedRef<SWidget> FHierarchyTableEditorToolkit::CreateTedsOutliner()
-{
-	using namespace UE::Editor::DataStorage::Queries;
-	using namespace UE::Editor::Outliner;
-
-	if (!AreEditorDataStorageFeaturesEnabled())
-	{
-		return SNew(STextBlock)
-			.Text(INVTEXT("You need to enable the Typed Element Data Storage plugin to see the table viewer!"));
-	}
-
-	if (!ensure(HierarchyTable->GetElementType()))
-	{
-		return SNullWidget::NullWidget;
-	}
-
-	FHierarchyTableEditorModule& HierarchyTableModule = FModuleManager::GetModuleChecked<FHierarchyTableEditorModule>("HierarchyTableEditor");
-
-	TArray<UScriptStruct*> HierarchyTableTypeColumns = HierarchyTableModule.GetElementTypeEditorColumns(HierarchyTable);
-	HierarchyTableTypeColumns.Add(FTypedElementOverrideColumn::StaticStruct());
-
-	FQueryDescription ColumnQueryDescription =
-		Select()
-			.ReadOnly(HierarchyTableTypeColumns)
-		.Compile();
-
-	InitialColumnQuery = GetMutableDataStorageFeature<IEditorDataStorageProvider>(StorageFeatureName)->RegisterQuery(MoveTemp(ColumnQueryDescription));
-
-	FSceneOutlinerInitializationOptions InitOptions;
-	InitOptions.bShowHeaderRow = true;
-	InitOptions.FilterBarOptions.bHasFilterBar = true;
-	InitOptions.OutlinerIdentifier = "HierarchyTableTedsOutliner";
-
-	FTedsOutlinerParams Params(nullptr);
-	{
-		FQueryDescription RowQueryDescription =
-			Select()
-			.Where()
-			.All<FTypedElementOverrideColumn>()
-			.Compile();
-
-		Params.QueryDescription = RowQueryDescription;
-		Params.CellWidgetPurposes = TArray<FName>{ TEXT("General.Cell") };
-		Params.HierarchyData = FTedsOutlinerHierarchyData::GetDefaultHierarchyData();
-		Params.bShowRowHandleColumn = false;
-	}
-
-	FTedsOutlinerModule& TedsOutlinerModule = FModuleManager::GetModuleChecked<FTedsOutlinerModule>("TedsOutliner");
-
-	IEditorDataStorageProvider* DSI = GetMutableDataStorageFeature<IEditorDataStorageProvider>(StorageFeatureName);
-	static TableHandle Table = DSI->FindTable(FName("Editor_HierarchyTableTable"));
-	
-	TArray<UScriptStruct*> BaseHierarchyTableTypeColumns = HierarchyTableModule.GetElementTypeEditorColumns(HierarchyTable);
-
-	for (int32 EntryIndex = 0; EntryIndex < HierarchyTable->GetTableData().Num(); ++EntryIndex)
-	{
-		const FHierarchyTableEntryData* const Entry = HierarchyTable->GetTableEntry(EntryIndex);
-
-		RowHandle Row = DSI->AddRow(Table);
-
-		FTypedElementOverrideColumn OverrideEntry;
-		OverrideEntry.OwnerEntryIndex = EntryIndex;
-		OverrideEntry.OwnerTable = HierarchyTable;
-		DSI->AddColumn(Row, MoveTemp(OverrideEntry));
-
-		DSI->AddColumn<FTypedElementLabelColumn>(Row, { .Label = Entry->Identifier.ToString() });
-
-		RowHandle* ParentRow = EntryIndexToHandleMap.Find(Entry->Parent);
-		if (ParentRow)
-		{
-			DSI->AddColumn<FTableRowParentColumn>(Row, { .Parent = *ParentRow });
-		}
-
-		for (const UScriptStruct* Column : BaseHierarchyTableTypeColumns)
-		{
-			DSI->AddColumn(Row, Column);
-		}
-
-		EntryIndexToHandleMap.Add(EntryIndex, Row);
-	}
-
-	TSharedRef<ISceneOutliner> TedsOutliner = TedsOutlinerModule.CreateTedsOutliner(InitOptions, Params, InitialColumnQuery);
-	TedsOutlinerPtr = TedsOutliner;
-
-	OnEntryAddedHandle = HierarchyTable->RegisterOnEntryAdded(UHierarchyTable::FOnEntryAdded::FDelegate::CreateRaw(this, &FHierarchyTableEditorToolkit::HandleOnEntryAdded));
-
-	return TedsOutliner;
-}
-
 void FHierarchyTableEditorToolkit::ExtendToolbar()
 {
 	FHierarchyTableEditorModule& HierarchyTableModule = FModuleManager::GetModuleChecked<FHierarchyTableEditorModule>("HierarchyTableEditor");
@@ -211,60 +99,11 @@ void FHierarchyTableEditorToolkit::ExtendToolbar()
 	static const FName MenuName = GetToolMenuToolbarName(ParentName);
 
 	UToolMenu* ToolMenu = UToolMenus::Get()->ExtendMenu(MenuName);
-	Handler->ExtendToolbar(ToolMenu, this);
-}
-
-int32 FHierarchyTableEditorToolkit::GetSelectedEntryIndex() const
-{
-	using namespace UE::Editor::DataStorage;
-	IEditorDataStorageProvider* DSI = GetMutableDataStorageFeature<IEditorDataStorageProvider>(StorageFeatureName);
-	
-	TArray<FSceneOutlinerTreeItemPtr> Selection = TedsOutlinerPtr->GetTree().GetSelectedItems();
-	if (Selection.Num() > 0)
-	{
-		const UE::Editor::Outliner::FTedsOutlinerTreeItem* TedsItem = Selection[0]->CastTo<UE::Editor::Outliner::FTedsOutlinerTreeItem>();
-		const UE::Editor::DataStorage::RowHandle ParentRowHandle = TedsItem->GetRowHandle();
-
-		const FTypedElementOverrideColumn* OverrideColumn = DSI->GetColumn<FTypedElementOverrideColumn>(ParentRowHandle);
-		return OverrideColumn->OwnerEntryIndex;
-	}
-
-	return INDEX_NONE;
+	Handler->ExtendToolbar(ToolMenu, *HierarchyTableWidget);
 }
 
 void FHierarchyTableEditorToolkit::HandleOnEntryAdded(const int32 EntryIndex)
 {
-	const FHierarchyTableEntryData* const EntryData = HierarchyTable->GetTableEntry(EntryIndex);
-
-	using namespace UE::Editor::DataStorage;
-	IEditorDataStorageProvider* DSI = GetMutableDataStorageFeature<IEditorDataStorageProvider>(StorageFeatureName);
-	static TableHandle Table = DSI->FindTable(FName("Editor_HierarchyTableTable"));
-
-	const RowHandle Row = DSI->AddRow(Table);
-	{
-		FTypedElementOverrideColumn OverrideEntry;
-		OverrideEntry.OwnerEntryIndex = EntryIndex;
-		OverrideEntry.OwnerTable = HierarchyTable;
-		DSI->AddColumn(Row, MoveTemp(OverrideEntry));
-	}
-
-	// Ideally would read the label directly from the HT entry struct instead of storing it itself
-	// but this is a built-in TEDS column type that is used the the tabel viewer widget.
-	DSI->AddColumn<FTypedElementLabelColumn>(Row, { .Label = EntryData->Identifier.ToString() });
-
-	const RowHandle* ParentRow = EntryIndexToHandleMap.Find(EntryData->Parent);
-	if (ParentRow)
-	{
-		DSI->AddColumn<FTableRowParentColumn>(Row, { .Parent = *ParentRow });
-	}
-
-	FHierarchyTableEditorModule& HierarchyTableModule = FModuleManager::GetModuleChecked<FHierarchyTableEditorModule>("HierarchyTableEditor");
-	for (const UScriptStruct* Column : HierarchyTableModule.GetElementTypeEditorColumns(HierarchyTable))
-	{
-		DSI->AddColumn(Row, Column);
-	}
-
-	EntryIndexToHandleMap.Add(EntryIndex, Row);
 }
 
 #undef LOCTEXT_NAMESPACE
\ No newline at end of file
diff --git a/Engine/Plugins/Experimental/Animation/HierarchyTable/Source/Editor/Private/HierarchyTableLabelWidget.cpp b/Engine/Plugins/Experimental/Animation/HierarchyTable/Source/Editor/Private/HierarchyTableLabelWidget.cpp
deleted file mode 100644
diff --git a/Engine/Plugins/Experimental/Animation/HierarchyTable/Source/Editor/Private/HierarchyTableTedsFactory.cpp b/Engine/Plugins/Experimental/Animation/HierarchyTable/Source/Editor/Private/HierarchyTableTedsFactory.cpp
deleted file mode 100644
diff --git a/Engine/Plugins/Experimental/Animation/HierarchyTable/Source/Editor/Private/HierarchyTableTypeHandler.cpp b/Engine/Plugins/Experimental/Animation/HierarchyTable/Source/Editor/Private/HierarchyTableTypeHandler.cpp
--- a/Engine/Plugins/Experimental/Animation/HierarchyTable/Source/Editor/Private/HierarchyTableTypeHandler.cpp
+++ b/Engine/Plugins/Experimental/Animation/HierarchyTable/Source/Editor/Private/HierarchyTableTypeHandler.cpp
@@ -2,18 +2,19 @@
 
 #include "HierarchyTableTypeHandler.h"
 #include "HierarchyTable.h"
+#include "Styling/SlateColor.h"
+#include "Textures/SlateIcon.h"
 #include "Widgets/Images/SImage.h"
 #include "Widgets/Layout/SSpacer.h"
 #include "Widgets/Text/STextBlock.h"
 #include "Widgets/SBoxPanel.h"
 
-TSharedRef<SWidget> UHierarchyTable_TableTypeHandler::CreateLabelWidget(const int32 EntryIndex) const
+FSlateIcon UHierarchyTable_TableTypeHandler::GetEntryIcon(const int32 EntryIndex) const
 {
-	return SNew(SHorizontalBox)
-		+SHorizontalBox::Slot()
-		.FillWidth(1.0f)
-		[
-			SNew(STextBlock)
-			.Text(FText::FromName(HierarchyTable->GetTableEntry(EntryIndex)->Identifier))
-		];
+	return FSlateIcon();
+}
+
+FSlateColor UHierarchyTable_TableTypeHandler::GetEntryIconColor(const int32 EntryIndex) const
+{
+	return FSlateColor::UseForeground();
 }
\ No newline at end of file
diff --git a/Engine/Plugins/Experimental/Animation/HierarchyTable/Source/Editor/Private/HierarchyTableWidgetConstructor.cpp b/Engine/Plugins/Experimental/Animation/HierarchyTable/Source/Editor/Private/HierarchyTableWidgetConstructor.cpp
deleted file mode 100644
diff --git a/Engine/Plugins/Experimental/Animation/HierarchyTable/Source/Editor/Private/SHierarchyTable.cpp b/Engine/Plugins/Experimental/Animation/HierarchyTable/Source/Editor/Private/SHierarchyTable.cpp
new file mode 100644
--- /dev/null
+++ b/Engine/Plugins/Experimental/Animation/HierarchyTable/Source/Editor/Private/SHierarchyTable.cpp
@@ -0,0 +1,208 @@
+// Copyright Epic Games, Inc. All Rights Reserved.
+
+#include "SHierarchyTable.h"
+
+#include "HierarchyTableEditorModule.h"
+#include "Modules/ModuleManager.h"
+#include "HierarchyTable.h"
+#include "Widgets/Views/STreeView.h"
+#include "SHierarchyTableRow.h"
+#include "IHierarchyTableColumn.h"
+
+#define LOCTEXT_NAMESPACE "SHierarchyTable"
+
+const FName SHierarchyTable::FColumns::IdentifierId("Identifier");
+const FName SHierarchyTable::FColumns::OverrideId("Override");
+
+void SHierarchyTable::Construct(const FArguments& InArgs, TObjectPtr<UHierarchyTable> InHierarchyTable)
+{
+	FHierarchyTableEditorModule& HierarchyTableModule = FModuleManager::GetModuleChecked<FHierarchyTableEditorModule>("HierarchyTableEditor");
+
+	TableHandler = HierarchyTableModule.CreateTableHandler(InHierarchyTable);
+
+	HierarchyTable = InHierarchyTable;
+
+	TSharedPtr<SHeaderRow> HeaderRow;
+
+	ChildSlot
+		[
+			SAssignNew(TreeView, STreeView<TSharedPtr<FTreeItem>>)
+				.TreeItemsSource(&RootItems)
+				.OnGenerateRow(this, &SHierarchyTable::TreeView_GenerateItemRow)
+				.OnGetChildren(this, &SHierarchyTable::TreeView_HandleGetChildren)
+				.OnContextMenuOpening(this, &SHierarchyTable::TreeView_HandleContextMenuOpening)
+				.OnItemScrolledIntoView(this, &SHierarchyTable::TreeView_OnItemScrolledIntoView)
+				.HighlightParentNodesForSelection(true)
+				.HeaderRow
+				(
+					SAssignNew(HeaderRow, SHeaderRow)
+					+ SHeaderRow::Column(SHierarchyTable::FColumns::IdentifierId)
+					.FillWidth(0.5f)
+					.DefaultLabel(LOCTEXT("IdentifierLabel", "Identifier"))
+					+ SHeaderRow::Column(SHierarchyTable::FColumns::OverrideId)
+					.FixedWidth(24.0f)
+					.HAlignHeader(HAlign_Center)
+					.VAlignHeader(VAlign_Center)
+					.HAlignCell(HAlign_Fill)
+					.VAlignCell(VAlign_Fill)
+					[
+						SNew(SBox)
+							.VAlign(VAlign_Center)
+							.HAlign(HAlign_Center)
+							[
+								SNew(SImage)
+									.ColorAndOpacity(FSlateColor::UseForeground())
+									.Image(FAppStyle::GetBrush("DetailsView.OverrideHere"))
+							]
+					]
+				)
+		];
+
+	TArray<TSharedPtr<IHierarchyTableColumn>> ElementTypeColumns = HierarchyTableModule.GetElementTypeEditorColumns(InHierarchyTable);
+
+	for (TSharedPtr<IHierarchyTableColumn> Column : ElementTypeColumns)
+	{
+		HeaderRow->AddColumn(SHeaderRow::Column(Column->GetColumnId())
+			.DefaultLabel(Column->GetColumnLabel())
+			.FillWidth(Column->GetColumnSize() * 0.5f));
+	}
+
+	RegenerateTreeViewItems();
+
+	// Expand all tree items on construction
+	for (TSharedPtr<FTreeItem> TreeItem : GetAllTreeItems())
+	{
+		TreeView->SetItemExpansion(TreeItem, true);
+	}
+}
+
+int32 SHierarchyTable::GetSelectedEntryIndex() const
+{
+	TArray<TSharedPtr<FTreeItem>> SelectedItems;
+	TreeView->GetSelectedItems(SelectedItems);
+
+	if (SelectedItems.Num() > 0)
+	{
+		return SelectedItems[0]->Index;
+	}
+
+	return INDEX_NONE;
+}
+
+void SHierarchyTable::RegenerateHierarchy()
+{
+	RegenerateTreeViewItems();
+}
+
+void SHierarchyTable::TreeView_HandleGetChildren(TSharedPtr<FTreeItem> InItem, TArray<TSharedPtr<FTreeItem>>& OutChildren) const
+{
+	OutChildren.Append(InItem->Children);
+}
+
+void SHierarchyTable::TreeView_OnItemScrolledIntoView(TSharedPtr<FTreeItem> InItem, const TSharedPtr<ITableRow>& InWidget)
+{
+	if (DeferredRenameRequest.IsValid())
+	{
+		DeferredRenameRequest->OnRenameRequested.ExecuteIfBound();
+		DeferredRenameRequest.Reset();
+	}
+}
+
+TSharedRef<ITableRow> SHierarchyTable::TreeView_GenerateItemRow(TSharedPtr<SHierarchyTable::FTreeItem> Item, const TSharedRef<STableViewBase>& OwnerTable)
+{
+	return SNew(SHierarchyTableRow, OwnerTable, SharedThis(this), Item);
+}
+
+TSharedPtr<SWidget> SHierarchyTable::TreeView_HandleContextMenuOpening()
+{
+	FMenuBuilder MenuBuilder(true, nullptr);
+
+	TArray<TSharedPtr<FTreeItem>> SelectedItems;
+	TreeView->GetSelectedItems(SelectedItems);
+
+	if (SelectedItems.Num() > 0)
+	{
+		TableHandler->ExtendContextMenu(MenuBuilder, *this);
+	}
+
+	return MenuBuilder.MakeWidget();
+}
+
+TArray<TSharedPtr<SHierarchyTable::FTreeItem>> SHierarchyTable::GetAllTreeItems()
+{
+	TArray<TSharedPtr<SHierarchyTable::FTreeItem>> AllItems;
+	AllItems.Append(RootItems);
+
+	for (int32 Index = 0; Index < AllItems.Num(); ++Index)
+	{
+		AllItems.Append(AllItems[Index]->Children);
+	}
+
+	return AllItems;
+}
+
+void SHierarchyTable::RegenerateTreeViewItems()
+{
+	// Make note of all tree items currently expanded
+	TSet<FName> ExpandedAttributeNames;
+	{
+		for (TSharedPtr<FTreeItem> TreeItem : GetAllTreeItems())
+		{
+			if (TreeView->IsItemExpanded(TreeItem))
+			{
+				ExpandedAttributeNames.Add(TreeItem->Name);
+			}
+		}
+	}
+
+	// Rebuild items
+	{
+		RootItems.Reset();
+
+		TMap<FName, TSharedPtr<FTreeItem>> ItemMap;
+
+		const TArray<FHierarchyTableEntryData>& TableData = HierarchyTable->GetTableData();
+
+		for (int32 EntryIndex = 0; EntryIndex < TableData.Num(); ++EntryIndex)
+		{
+			const FHierarchyTableEntryData& Entry = TableData[EntryIndex];
+
+			TSharedPtr<FTreeItem> Item = MakeShared<FTreeItem>();
+			Item->Name = Entry.Identifier;
+			Item->Index = EntryIndex;
+
+			if (Entry.Parent == INDEX_NONE)
+			{
+				RootItems.Add(Item);
+			}
+			else
+			{
+				const FHierarchyTableEntryData* ParentEntry = HierarchyTable->GetTableEntry(Entry.Parent);
+				check(ParentEntry);
+
+				const TSharedPtr<FTreeItem>* ParentTreeItem = ItemMap.Find(ParentEntry->Identifier);
+				check(ParentTreeItem);
+
+				(*ParentTreeItem)->Children.Add(Item);
+			}
+
+			ItemMap.Add(Item->Name, Item);
+		}
+	}
+
+	// Update tree view and restore tree item expanded states
+	{
+		check(TreeView);
+		TreeView->RequestTreeRefresh();
+
+		for (TSharedPtr<FTreeItem> TreeItem : GetAllTreeItems())
+		{
+			if (ExpandedAttributeNames.Contains(TreeItem->Name))
+			{
+				TreeView->SetItemExpansion(TreeItem, true);
+			}
+		}
+	}
+}
+
+#undef LOCTEXT_NAMESPACE
\ No newline at end of file
diff --git a/Engine/Plugins/Experimental/Animation/HierarchyTable/Source/Editor/Private/SHierarchyTable.h b/Engine/Plugins/Experimental/Animation/HierarchyTable/Source/Editor/Private/SHierarchyTable.h
new file mode 100644
--- /dev/null
+++ b/Engine/Plugins/Experimental/Animation/HierarchyTable/Source/Editor/Private/SHierarchyTable.h
@@ -0,0 +1,70 @@
+// Copyright Epic Games, Inc. All Rights Reserved.
+
+#pragma once
+
+#include "CoreMinimal.h"
+#include "IHierarchyTable.h"
+#include "Widgets/SCompoundWidget.h"
+#include "Widgets/Views/STableViewBase.h"
+#include "Widgets/Views/STreeView.h"
+
+class ITableRow;
+class SHierarchyTableRow;
+class UHierarchyTable;
+class UHierarchyTable_TableTypeHandler;
+
+class SHierarchyTable : public SCompoundWidget, public IHierarchyTable
+{
+public:
+	friend SHierarchyTableRow;
+
+	SLATE_BEGIN_ARGS(SHierarchyTable) {}
+	SLATE_END_ARGS()
+
+	struct FColumns
+	{
+		static const FName IdentifierId;
+		static const FName OverrideId;
+	};
+
+	struct FTreeItem
+	{
+		int32 Index;
+		FName Name;
+		TArray<TSharedPtr<FTreeItem>> Children;
+
+		DECLARE_DELEGATE(FOnRenameRequested);
+		FOnRenameRequested OnRenameRequested;
+	};
+
+	void Construct(const FArguments& InArgs, TObjectPtr<UHierarchyTable> InHierarchyTable);
+
+	// IHierarchyTable
+	int32 GetSelectedEntryIndex() const override;
+	void RegenerateHierarchy() override;
+
+private:
+	void TreeView_HandleGetChildren(TSharedPtr<FTreeItem> InItem, TArray<TSharedPtr<FTreeItem>>& OutChildren) const;
+
+	void TreeView_OnItemScrolledIntoView(TSharedPtr<FTreeItem> InItem, const TSharedPtr<ITableRow>& InWidget);
+
+	TSharedRef<ITableRow> TreeView_GenerateItemRow(TSharedPtr<FTreeItem> Item, const TSharedRef<STableViewBase>& OwnerTable);
+
+	TSharedPtr<SWidget> TreeView_HandleContextMenuOpening();
+
+private:
+	TArray<TSharedPtr<FTreeItem>> GetAllTreeItems();
+
+	void RegenerateTreeViewItems();
+
+private:
+	TSharedPtr<STreeView<TSharedPtr<FTreeItem>>> TreeView;
+
+	TArray<TSharedPtr<FTreeItem>> RootItems;
+
+	TSharedPtr<FTreeItem> DeferredRenameRequest;
+
+	TObjectPtr<UHierarchyTable> HierarchyTable;
+
+	TObjectPtr<UHierarchyTable_TableTypeHandler> TableHandler;
+};
diff --git a/Engine/Plugins/Experimental/Animation/HierarchyTable/Source/Editor/Private/SHierarchyTableRow.cpp b/Engine/Plugins/Experimental/Animation/HierarchyTable/Source/Editor/Private/SHierarchyTableRow.cpp
new file mode 100644
--- /dev/null
+++ b/Engine/Plugins/Experimental/Animation/HierarchyTable/Source/Editor/Private/SHierarchyTableRow.cpp
@@ -0,0 +1,165 @@
+// Copyright Epic Games, Inc. All Rights Reserved.
+
+#pragma once
+
+#include "SHierarchyTableRow.h"
+#include "Widgets/Text/SInlineEditableTextBlock.h"
+#include "HierarchyTableEditorModule.h"
+#include "IHierarchyTableColumn.h"
+#include "HierarchyTableTypeHandler.h"
+#include "Widgets/Images/SImage.h"
+#include "Widgets/Input/SButton.h"
+#include "Modules/ModuleManager.h"
+#include "ScopedTransaction.h"
+
+#define LOCTEXT_NAMESPACE "SHierarchyTableRow"
+
+FReply SHierarchyTableRow::OnDragDetected(const FGeometry& MyGeometry, const FPointerEvent& MouseEvent)
+{
+	return FReply::Unhandled();
+}
+
+TOptional<EItemDropZone> SHierarchyTableRow::OnCanAcceptDrop(const FDragDropEvent& DragDropEvent, EItemDropZone DropZone, TSharedPtr<SHierarchyTable::FTreeItem> TargetItem)
+{
+	TOptional<EItemDropZone> ReturnedDropZone;
+
+	return ReturnedDropZone;
+}
+
+FReply SHierarchyTableRow::OnAcceptDrop(const FDragDropEvent& DragDropEvent, EItemDropZone DropZone, TSharedPtr<SHierarchyTable::FTreeItem> TargetItem)
+{
+	return FReply::Handled();
+}
+
+void SHierarchyTableRow::Construct(const FArguments& InArgs, const TSharedRef<STableViewBase>& InOwnerTableView, TSharedPtr<SHierarchyTable> InHierarchyTableWidget, TSharedPtr<SHierarchyTable::FTreeItem> InTreeItem)
+{
+	HierarchyTableWidget = InHierarchyTableWidget;
+	TreeItem = InTreeItem;
+
+	OnRenamed = InArgs._OnRenamed;
+	OnReparented = InArgs._OnReparented;
+
+	FSuperRowType::Construct(
+		FSuperRowType::FArguments()
+		.OnDragDetected(this, &SHierarchyTableRow::OnDragDetected)
+		.OnCanAcceptDrop(this, &SHierarchyTableRow::OnCanAcceptDrop)
+		.OnAcceptDrop(this, &SHierarchyTableRow::OnAcceptDrop),
+		InOwnerTableView);
+}
+
+void SHierarchyTableRow::OnCommitRename(const FText& InText, ETextCommit::Type CommitInfo)
+{
+	const FName OldName = TreeItem->Name;
+	const FName NewName = FName(InText.ToString());
+
+	if (OnRenamed.IsBound())
+	{
+		const bool bSuccess = OnRenamed.Execute(OldName, NewName);
+		if (bSuccess)
+		{
+			TreeItem->Name = NewName;
+		}
+	}
+}
+
+TSharedRef<SWidget> SHierarchyTableRow::GenerateWidgetForColumn(const FName& ColumnName)
+{
+	if (ColumnName == SHierarchyTable::FColumns::IdentifierId)
+	{
+		TSharedPtr<SInlineEditableTextBlock> InlineWidget;
+
+		TSharedRef<SHorizontalBox> HorizontalBox = SNew(SHorizontalBox)
+			+ SHorizontalBox::Slot()
+			.AutoWidth()
+			[
+				SNew(SExpanderArrow, SharedThis(this))
+					.ShouldDrawWires(true)
+			]
+			+ SHorizontalBox::Slot()
+			.Padding(2.0f, 2.0f, 4.0f, 2.0f)
+			.AutoWidth()
+			[
+				SNew(SImage)
+					.Image(HierarchyTableWidget->TableHandler->GetEntryIcon(TreeItem->Index).GetSmallIcon())
+					.ColorAndOpacity(HierarchyTableWidget->TableHandler->GetEntryIconColor(TreeItem->Index))
+			]
+			+ SHorizontalBox::Slot()
+			.AutoWidth()
+			[
+				SAssignNew(InlineWidget, SInlineEditableTextBlock)
+					.Text_Lambda([this]()
+						{
+							return FText::FromName(TreeItem->Name);
+						})
+					.OnTextCommitted(this, &SHierarchyTableRow::OnCommitRename)
+			];
+
+		TreeItem->OnRenameRequested.BindSP(InlineWidget.Get(), &SInlineEditableTextBlock::EnterEditingMode);
+
+		return HorizontalBox;
+	}
+
+	if (ColumnName == SHierarchyTable::FColumns::OverrideId)
+	{
+		const bool bHasParent = HierarchyTableWidget->HierarchyTable->GetTableEntry(TreeItem->Index)->HasParent();
+
+		return SNew(SButton)
+			.ButtonStyle(FAppStyle::Get(), "SimpleButton")
+			.IsEnabled(bHasParent)
+			.OnClicked_Lambda([this]()
+				{
+					const FScopedTransaction Transaction(LOCTEXT("ToggleOverride", "Toggle Override"));
+					HierarchyTableWidget->HierarchyTable->Modify();
+					HierarchyTableWidget->HierarchyTable->GetMutableTableEntry(TreeItem->Index)->ToggleOverridden();
+					return FReply::Handled();
+				})
+			.ContentPadding(2.0f)
+			[
+				SNew(SImage)
+					.Image_Lambda([this]()
+						{
+							const FHierarchyTableEntryData* const EntryData = HierarchyTableWidget->HierarchyTable->GetTableEntry(TreeItem->Index);
+							const bool bHasOverriddenChildren = EntryData->HasOverriddenChildren();
+
+							if (EntryData->IsOverridden())
+							{
+								if (bHasOverriddenChildren)
+								{
+									return FAppStyle::GetBrush(TEXT("DetailsView.OverrideHereInside"));
+								}
+								else
+								{
+									return FAppStyle::GetBrush(TEXT("DetailsView.OverrideHere"));
+								}
+							}
+							else
+							{
+								if (bHasOverriddenChildren)
+								{
+									return FAppStyle::GetBrush(TEXT("DetailsView.OverrideInside"));
+								}
+								else
+								{
+									return FAppStyle::GetBrush(TEXT("DetailsView.OverrideNone"));
+								}
+							}
+						})
+					.ColorAndOpacity(FSlateColor::UseForeground())
+			];
+	}
+
+	FHierarchyTableEditorModule& HierarchyTableModule = FModuleManager::GetModuleChecked<FHierarchyTableEditorModule>("HierarchyTableEditor");
+	const TArray<TSharedPtr<IHierarchyTableColumn>> Columns = HierarchyTableModule.GetElementTypeEditorColumns(HierarchyTableWidget->HierarchyTable);
+
+	for (TSharedPtr<IHierarchyTableColumn> Column : Columns)
+	{
+		if (ColumnName == Column->GetColumnId())
+		{
+			return Column->CreateEntryWidget(HierarchyTableWidget->HierarchyTable, TreeItem->Index);
+		}
+	}
+
+	return SNullWidget::NullWidget;
+}
+
+#undef LOCTEXT_NAMESPACE
\ No newline at end of file
diff --git a/Engine/Plugins/Experimental/Animation/HierarchyTable/Source/Editor/Private/SHierarchyTableRow.h b/Engine/Plugins/Experimental/Animation/HierarchyTable/Source/Editor/Private/SHierarchyTableRow.h
new file mode 100644
--- /dev/null
+++ b/Engine/Plugins/Experimental/Animation/HierarchyTable/Source/Editor/Private/SHierarchyTableRow.h
@@ -0,0 +1,39 @@
+// Copyright Epic Games, Inc. All Rights Reserved.
+
+#pragma once
+
+#include "CoreMinimal.h"
+#include "Framework/MultiBox/MultiBoxBuilder.h"
+#include "SHierarchyTable.h"
+#include "Widgets/Views/STableRow.h"
+#include "Widgets/Views/STableViewBase.h"
+
+class SHierarchyTableRow : public SMultiColumnTableRow<TSharedPtr<SHierarchyTable::FTreeItem>>
+{
+public:
+	DECLARE_DELEGATE_RetVal_TwoParams(bool /* Success */, FOnRenamed, FName /* Old Name */, FName /* New Name */);
+	DECLARE_DELEGATE_TwoParams(FOnReparented, FName /* Attribute Name */, FName /* New Parent Name */);
+
+	SLATE_BEGIN_ARGS(SHierarchyTableRow) {}
+		SLATE_EVENT(FOnRenamed, OnRenamed);
+		SLATE_EVENT(FOnReparented, OnReparented);
+	SLATE_END_ARGS()
+
+	FReply OnDragDetected(const FGeometry& MyGeometry, const FPointerEvent& MouseEvent);
+
+	TOptional<EItemDropZone> OnCanAcceptDrop(const FDragDropEvent& DragDropEvent, EItemDropZone DropZone, TSharedPtr<SHierarchyTable::FTreeItem> TargetItem);
+
+	FReply OnAcceptDrop(const FDragDropEvent& DragDropEvent, EItemDropZone DropZone, TSharedPtr<SHierarchyTable::FTreeItem> TargetItem);
+
+	void Construct(const FArguments& InArgs, const TSharedRef<STableViewBase>& InOwnerTableView, TSharedPtr<SHierarchyTable> InHierarchyTableWidget, TSharedPtr<SHierarchyTable::FTreeItem> InTreeItem);
+
+	void OnCommitRename(const FText& InText, ETextCommit::Type CommitInfo);
+
+	TSharedRef<SWidget> GenerateWidgetForColumn(const FName& ColumnName) override;
+
+private:
+	TSharedPtr<SHierarchyTable> HierarchyTableWidget;
+	TSharedPtr<SHierarchyTable::FTreeItem> TreeItem;
+	FOnRenamed OnRenamed;
+	FOnReparented OnReparented;
+};
\ No newline at end of file
diff --git a/Engine/Plugins/Experimental/Animation/HierarchyTable/Source/Editor/Public/HierarchyTable/Columns/OverrideColumn.h b/Engine/Plugins/Experimental/Animation/HierarchyTable/Source/Editor/Public/HierarchyTable/Columns/OverrideColumn.h
deleted file mode 100644
diff --git a/Engine/Plugins/Experimental/Animation/HierarchyTable/Source/Editor/Public/HierarchyTableEditorModule.h b/Engine/Plugins/Experimental/Animation/HierarchyTable/Source/Editor/Public/HierarchyTableEditorModule.h
--- a/Engine/Plugins/Experimental/Animation/HierarchyTable/Source/Editor/Public/HierarchyTableEditorModule.h
+++ b/Engine/Plugins/Experimental/Animation/HierarchyTable/Source/Editor/Public/HierarchyTableEditorModule.h
@@ -7,6 +7,8 @@
 #include "Modules/ModuleInterface.h"
 #include "HierarchyTable.h"
 
+struct IHierarchyTableColumn;
+
 class FHierarchyTableEditorModule : public IModuleInterface
 {
 public:
@@ -22,14 +24,14 @@ public:
 
 	HIERARCHYTABLEEDITOR_API const TObjectPtr<UHierarchyTable_TableTypeHandler> CreateTableHandler(const TObjectPtr<const UScriptStruct> TableType);
 
-	HIERARCHYTABLEEDITOR_API void RegisterElementTypeEditorColumns(const UScriptStruct* ElementType, const TArray<TObjectPtr<UScriptStruct>> Columns);
+	HIERARCHYTABLEEDITOR_API void RegisterElementTypeEditorColumns(const UScriptStruct* ElementType, const TArray<TSharedPtr<IHierarchyTableColumn>> Columns);
 
 	HIERARCHYTABLEEDITOR_API void UnregisterElementTypeEditorColumns(const UScriptStruct* ElementType);
 
-	TArray<UScriptStruct*> GetElementTypeEditorColumns(const TObjectPtr<const UHierarchyTable> HierarchyTable);
+	TArray<TSharedPtr<IHierarchyTableColumn>> GetElementTypeEditorColumns(const TObjectPtr<const UHierarchyTable> HierarchyTable);
 
 private:
 	TMap<const UScriptStruct*, const UClass*> TableHandlers;
 
-	TMap<const UScriptStruct*, TArray<UScriptStruct*>> EditorColumns;
+	TMap<TObjectPtr<const UScriptStruct>, TArray<TSharedPtr<IHierarchyTableColumn>>> EditorColumns;
 };
\ No newline at end of file
diff --git a/Engine/Plugins/Experimental/Animation/HierarchyTable/Source/Editor/Public/HierarchyTableEditorToolkit.h b/Engine/Plugins/Experimental/Animation/HierarchyTable/Source/Editor/Public/HierarchyTableEditorToolkit.h
--- a/Engine/Plugins/Experimental/Animation/HierarchyTable/Source/Editor/Public/HierarchyTableEditorToolkit.h
+++ b/Engine/Plugins/Experimental/Animation/HierarchyTable/Source/Editor/Public/HierarchyTableEditorToolkit.h
@@ -5,15 +5,15 @@
 #include "CoreMinimal.h"
 #include "HierarchyTable.h"
 #include "Toolkits/AssetEditorToolkit.h"
-#include "Elements/Common/TypedElementHandles.h"
-#include "IHierarchyTableEditor.h"
 
-class FHierarchyTableEditorToolkit : public FAssetEditorToolkit, public IHierarchyTableEditor
+class SHierarchyTable;
+
+class FHierarchyTableEditorToolkit : public FAssetEditorToolkit
 {
 public:
 	void InitEditor(const TArray<UObject*>& InObjects);
-	virtual void OnClose() override;
 
+	// Begin FAssetEditorToolkit
 	void RegisterTabSpawners(const TSharedRef<class FTabManager>& TabManager) override;
 	void UnregisterTabSpawners(const TSharedRef<class FTabManager>& TabManager) override;
 
@@ -21,25 +21,13 @@ public:
 	FText GetBaseToolkitName() const override { return INVTEXT("Hierarchy Table Editor"); }
 	FString GetWorldCentricTabPrefix() const override { return "Hierarchy Table "; }
 	FLinearColor GetWorldCentricTabColorScale() const override { return {}; }
+	// End FAssetEditorToolkit
 
 	void ExtendToolbar();
-
-	// Begin IHierarchyTableEditor
-	virtual int32 GetSelectedEntryIndex() const;
-	// End IHierarchyTableEditor
-
 private:
 	void HandleOnEntryAdded(const int32 EntryIndex);
 
-	TSharedRef<SWidget> CreateTedsOutliner();
-
-	UHierarchyTable* HierarchyTable;
-
-	TMap<int32, UE::Editor::DataStorage::RowHandle> EntryIndexToHandleMap;
-
-	UE::Editor::DataStorage::QueryHandle InitialColumnQuery;
-
-	TSharedPtr<class ISceneOutliner> TedsOutlinerPtr;
+	TObjectPtr<UHierarchyTable> HierarchyTable;
 
-	FDelegateHandle OnEntryAddedHandle;
+	TSharedPtr<SHierarchyTable> HierarchyTableWidget;
 };
diff --git a/Engine/Plugins/Experimental/Animation/HierarchyTable/Source/Editor/Public/HierarchyTableLabelWidget.h b/Engine/Plugins/Experimental/Animation/HierarchyTable/Source/Editor/Public/HierarchyTableLabelWidget.h
deleted file mode 100644
diff --git a/Engine/Plugins/Experimental/Animation/HierarchyTable/Source/Editor/Public/HierarchyTableProxyColumn.h b/Engine/Plugins/Experimental/Animation/HierarchyTable/Source/Editor/Public/HierarchyTableProxyColumn.h
deleted file mode 100644
diff --git a/Engine/Plugins/Experimental/Animation/HierarchyTable/Source/Editor/Public/HierarchyTableTedsFactory.h b/Engine/Plugins/Experimental/Animation/HierarchyTable/Source/Editor/Public/HierarchyTableTedsFactory.h
deleted file mode 100644
diff --git a/Engine/Plugins/Experimental/Animation/HierarchyTable/Source/Editor/Public/HierarchyTableTypeHandler.h b/Engine/Plugins/Experimental/Animation/HierarchyTable/Source/Editor/Public/HierarchyTableTypeHandler.h
--- a/Engine/Plugins/Experimental/Animation/HierarchyTable/Source/Editor/Public/HierarchyTableTypeHandler.h
+++ b/Engine/Plugins/Experimental/Animation/HierarchyTable/Source/Editor/Public/HierarchyTableTypeHandler.h
@@ -5,13 +5,16 @@
 #include "UObject/Object.h"
 #include "StructUtils/InstancedStruct.h"
 #include "HierarchyTableType.h"
+#include "IHierarchyTable.h"
 #include "HierarchyTableTypeHandler.generated.h"
 
 class UHierarchyTable;
 class UScriptStruct;
 class UToolMenu;
 class SWidget;
-class IHierarchyTableEditor;
+class FMenuBuilder;
+struct FSlateIcon;
+struct FSlateColor;
 
 UCLASS(Abstract)
 class HIERARCHYTABLEEDITOR_API UHierarchyTable_TableTypeHandler : public UObject
@@ -25,9 +28,13 @@ public:
 
 	virtual bool FactoryConfigureProperties(FInstancedStruct& TableType) const { return false; }
 
-	virtual void ExtendToolbar(UToolMenu* ToolMenu, const IHierarchyTableEditor* Editor) const {}
+	virtual void ExtendToolbar(UToolMenu* ToolMenu, IHierarchyTable& HierarchyTableView) const {}
 
-	virtual TSharedRef<SWidget> CreateLabelWidget(const int32 EntryIndex) const;
+	virtual void ExtendContextMenu(FMenuBuilder& MenuBuilder, IHierarchyTable& HierarchyTableView) const {}
+
+	virtual FSlateIcon GetEntryIcon(const int32 EntryIndex) const;
+
+	virtual FSlateColor GetEntryIconColor(const int32 EntryIndex) const;
 
 	virtual void ConstructHierarchy() {};
 
diff --git a/Engine/Plugins/Experimental/Animation/HierarchyTable/Source/Editor/Public/HierarchyTableWidgetConstructor.h b/Engine/Plugins/Experimental/Animation/HierarchyTable/Source/Editor/Public/HierarchyTableWidgetConstructor.h
deleted file mode 100644
diff --git a/Engine/Plugins/Experimental/Animation/HierarchyTable/Source/Editor/Public/IHierarchyTable.h b/Engine/Plugins/Experimental/Animation/HierarchyTable/Source/Editor/Public/IHierarchyTable.h
new file mode 100644
--- /dev/null
+++ b/Engine/Plugins/Experimental/Animation/HierarchyTable/Source/Editor/Public/IHierarchyTable.h
@@ -0,0 +1,13 @@
+// Copyright Epic Games, Inc. All Rights Reserved.
+
+#pragma once
+
+#include "CoreMinimal.h"
+
+class IHierarchyTable
+{
+public:
+	virtual int32 GetSelectedEntryIndex() const = 0;
+
+	virtual void RegenerateHierarchy() = 0;
+};
\ No newline at end of file
diff --git a/Engine/Plugins/Experimental/Animation/HierarchyTable/Source/Editor/Public/IHierarchyTableColumn.h b/Engine/Plugins/Experimental/Animation/HierarchyTable/Source/Editor/Public/IHierarchyTableColumn.h
new file mode 100644
--- /dev/null
+++ b/Engine/Plugins/Experimental/Animation/HierarchyTable/Source/Editor/Public/IHierarchyTableColumn.h
@@ -0,0 +1,20 @@
+// Copyright Epic Games, Inc. All Rights Reserved.
+
+#pragma once
+
+#include "CoreMinimal.h"
+#include "Widgets/SWidget.h"
+
+class UHierarchyTable;
+
+struct HIERARCHYTABLEEDITOR_API IHierarchyTableColumn
+{
+public:
+	virtual FName GetColumnId() const = 0;
+	virtual FText GetColumnLabel() const = 0;
+	virtual float GetColumnSize() const = 0;
+
+	virtual TSharedRef<SWidget> CreateEntryWidget(TObjectPtr<UHierarchyTable> HierarchyTable, int32 EntryIndex) = 0;
+
+	virtual TSharedRef<SWidget> CreateHeaderWidget() = 0;
+};
\ No newline at end of file
diff --git a/Engine/Plugins/Experimental/Animation/HierarchyTable/Source/Editor/Public/IHierarchyTableEditor.h b/Engine/Plugins/Experimental/Animation/HierarchyTable/Source/Editor/Public/IHierarchyTableEditor.h
deleted file mode 100644
diff --git a/Engine/Plugins/Experimental/Animation/HierarchyTable/Source/Runtime/Private/HierarchyTable.cpp b/Engine/Plugins/Experimental/Animation/HierarchyTable/Source/Runtime/Private/HierarchyTable.cpp
--- a/Engine/Plugins/Experimental/Animation/HierarchyTable/Source/Runtime/Private/HierarchyTable.cpp
+++ b/Engine/Plugins/Experimental/Animation/HierarchyTable/Source/Runtime/Private/HierarchyTable.cpp
@@ -97,6 +97,12 @@ bool UHierarchyTable::HasIdentifier(const FName Identifier) const
 
 int32 UHierarchyTable::AddEntry(const FHierarchyTableEntryData& Entry)
 {
+	// Do not allow entries with duplicate identifiers
+	if (HasIdentifier(Entry.Identifier))
+	{
+		return INDEX_NONE;
+	}
+
 	const int32 EntryIndex = TableData.Add(Entry);
 #if WITH_EDITOR
 	OnEntryAdded.Broadcast(EntryIndex);
@@ -104,6 +110,31 @@ int32 UHierarchyTable::AddEntry(const FHierarchyTableEntryData& Entry)
 	return EntryIndex;
 }
 
+void UHierarchyTable::RemoveEntry(const int32 IndexToRemove)
+{
+	if (!TableData.IsValidIndex(IndexToRemove))
+	{
+		return;
+	}
+	FHierarchyTableEntryData& EntryToRemove = TableData[IndexToRemove];
+
+	for (int32 EntryIndex = 0; EntryIndex < TableData.Num(); ++EntryIndex)
+	{
+		FHierarchyTableEntryData& Entry = TableData[EntryIndex];
+
+		if (Entry.Parent == IndexToRemove)
+		{
+			Entry.Parent = EntryToRemove.Parent;
+		}
+		else if (Entry.Parent > IndexToRemove)
+		{
+			--Entry.Parent;
+		}
+	}
+
+	TableData.RemoveAt(IndexToRemove);
+}
+
 #if WITH_EDITOR
 FDelegateHandle UHierarchyTable::RegisterOnEntryAdded(const FOnEntryAdded::FDelegate& Delegate)
 {
diff --git a/Engine/Plugins/Experimental/Animation/HierarchyTable/Source/Runtime/Public/HierarchyTable.h b/Engine/Plugins/Experimental/Animation/HierarchyTable/Source/Runtime/Public/HierarchyTable.h
--- a/Engine/Plugins/Experimental/Animation/HierarchyTable/Source/Runtime/Public/HierarchyTable.h
+++ b/Engine/Plugins/Experimental/Animation/HierarchyTable/Source/Runtime/Public/HierarchyTable.h
@@ -162,6 +162,8 @@ public:
 
 	int32 AddEntry(const FHierarchyTableEntryData& Entry);
 
+	void RemoveEntry(const int32 Index);
+
 	TArray<const FHierarchyTableEntryData*> GetChildren(const FHierarchyTableEntryData& Parent) const;
 
 	bool HasIdentifier(const FName Identifier) const;
diff --git a/Engine/Plugins/Experimental/Animation/HierarchyTableAnimation/HierarchyTableAnimation.uplugin b/Engine/Plugins/Experimental/Animation/HierarchyTableAnimation/HierarchyTableAnimation.uplugin
--- a/Engine/Plugins/Experimental/Animation/HierarchyTableAnimation/HierarchyTableAnimation.uplugin
+++ b/Engine/Plugins/Experimental/Animation/HierarchyTableAnimation/HierarchyTableAnimation.uplugin
@@ -36,10 +36,6 @@
 		{
 			"Name": "HierarchyTable",
 			"Enabled": true
-		},
-		{
-			"Name": "EditorDataStorageFeatures",
-			"Enabled": true
 		}
 	]
 }
diff --git a/Engine/Plugins/Experimental/Animation/HierarchyTableAnimation/Source/Editor/HierarchyTableAnimationEditor.Build.cs b/Engine/Plugins/Experimental/Animation/HierarchyTableAnimation/Source/Editor/HierarchyTableAnimationEditor.Build.cs
--- a/Engine/Plugins/Experimental/Animation/HierarchyTableAnimation/Source/Editor/HierarchyTableAnimationEditor.Build.cs
+++ b/Engine/Plugins/Experimental/Animation/HierarchyTableAnimation/Source/Editor/HierarchyTableAnimationEditor.Build.cs
@@ -18,7 +18,6 @@ public class HierarchyTableAnimationEditor : ModuleRules
 			"HierarchyTableRuntime",
 			"HierarchyTableEditor",
 			"HierarchyTableAnimationRuntime",
-			"TypedElementFramework", // TODO: Remove
 			"ToolMenus"
 		);
 	}
diff --git a/Engine/Plugins/Experimental/Animation/HierarchyTableAnimation/Source/Editor/Private/HierachyTableAnimationEditorModule.cpp b/Engine/Plugins/Experimental/Animation/HierarchyTableAnimation/Source/Editor/Private/HierachyTableAnimationEditorModule.cpp
--- a/Engine/Plugins/Experimental/Animation/HierarchyTableAnimation/Source/Editor/Private/HierachyTableAnimationEditorModule.cpp
+++ b/Engine/Plugins/Experimental/Animation/HierarchyTableAnimation/Source/Editor/Private/HierachyTableAnimationEditorModule.cpp
@@ -7,8 +7,8 @@
 #include "HierarchyTableEditorModule.h"
 #include "SkeletonHierarchyTableType.h"
 #include "SkeletonHierarchyTableTypeHandler.h"
-#include "MaskProfile/MaskProfileProxyColumn.h"
-#include "TimeProfile/TimeProfileProxyColumn.h"
+#include "MaskProfile/MaskProfileColumn.h"
+#include "TimeProfile/TimeProfileColumn.h"
 
 void FHierarchyTableAnimationEditorModule::StartupModule()
 {
@@ -17,10 +17,17 @@ void FHierarchyTableAnimationEditorModule::StartupModule()
 	HierarchyTableModule.RegisterTableType(FHierarchyTable_TableType_Skeleton::StaticStruct(), UHierarchyTable_TableTypeHandler_Skeleton::StaticClass());
 
 	HierarchyTableModule.RegisterElementTypeEditorColumns(FHierarchyTable_ElementType_Mask::StaticStruct(),
-		{ FHierarchyTableMaskColumn_Value::StaticStruct() });
+		{
+			MakeShared<FHierarchyTableColumn_Mask>(),
+		});
 
 	HierarchyTableModule.RegisterElementTypeEditorColumns(FHierarchyTable_ElementType_Time::StaticStruct(),
-		{ FHierarchyTableTimeColumn_StartTime::StaticStruct(), FHierarchyTableTimeColumn_EndTime::StaticStruct(), FHierarchyTableTimeColumn_TimeFactor::StaticStruct() });
+		{
+			MakeShared<FHierarchyTableColumn_TimeProfile_StartTime>(),
+			MakeShared<FHierarchyTableColumn_TimeProfile_EndTime>(),
+			MakeShared<FHierarchyTableColumn_TimeProfile_TimeFactor>(),
+			MakeShared<FHierarchyTableColumn_TimeProfile_Preview>(),
+		});
 }
 
 void FHierarchyTableAnimationEditorModule::ShutdownModule()
diff --git a/Engine/Plugins/Experimental/Animation/HierarchyTableAnimation/Source/Editor/Private/MaskProfile/MaskProfileColumn.cpp b/Engine/Plugins/Experimental/Animation/HierarchyTableAnimation/Source/Editor/Private/MaskProfile/MaskProfileColumn.cpp
new file mode 100644
--- /dev/null
+++ b/Engine/Plugins/Experimental/Animation/HierarchyTableAnimation/Source/Editor/Private/MaskProfile/MaskProfileColumn.cpp
@@ -0,0 +1,42 @@
+// Copyright Epic Games, Inc. All Rights Reserved.
+
+#include "MaskProfile/MaskProfileColumn.h"
+#include "Widgets/Input/SSpinBox.h"
+#include "MaskProfile/HierarchyTableTypeMask.h"
+#include "HierarchyTable.h"
+#include "Editor.h"
+
+#define LOCTEXT_NAMESPACE "FHierarchyTableColumn_Mask"
+
+TSharedRef<SWidget> FHierarchyTableColumn_Mask::CreateEntryWidget(TObjectPtr<UHierarchyTable> HierarchyTable, int32 EntryIndex)
+{
+	return SNew(SSpinBox<float>)
+		.IsEnabled_Lambda([HierarchyTable, EntryIndex]() { return HierarchyTable->GetTableEntry(EntryIndex)->IsOverridden(); })
+		.MinDesiredWidth(100.0f)
+		.MinValue(0.0f)
+		.MaxValue(1.0f)
+		.Value_Lambda([HierarchyTable, EntryIndex]()
+			{
+				return HierarchyTable->GetTableEntry(EntryIndex)->GetValue<FHierarchyTable_ElementType_Mask>()->Value;
+			})
+		.OnValueChanged_Lambda([HierarchyTable, EntryIndex](float NewValue)
+			{
+				HierarchyTable->GetMutableTableEntry(EntryIndex)->GetMutableValue<FHierarchyTable_ElementType_Mask>()->Value = NewValue;
+			})
+		.OnBeginSliderMovement_Lambda([HierarchyTable]()
+			{
+				GEditor->BeginTransaction(LOCTEXT("SetMaskValue", "Set Mask Value"));
+				HierarchyTable->Modify();
+			})
+		.OnEndSliderMovement_Lambda([](float)
+			{
+				GEditor->EndTransaction();
+			});
+}
+
+TSharedRef<SWidget> FHierarchyTableColumn_Mask::CreateHeaderWidget()
+{
+	return SNullWidget::NullWidget;
+}
+
+#undef LOCTEXT_NAMESPACE
\ No newline at end of file
diff --git a/Engine/Plugins/Experimental/Animation/HierarchyTableAnimation/Source/Editor/Private/MaskProfile/MaskProfileColumn.h b/Engine/Plugins/Experimental/Animation/HierarchyTableAnimation/Source/Editor/Private/MaskProfile/MaskProfileColumn.h
new file mode 100644
--- /dev/null
+++ b/Engine/Plugins/Experimental/Animation/HierarchyTableAnimation/Source/Editor/Private/MaskProfile/MaskProfileColumn.h
@@ -0,0 +1,20 @@
+// Copyright Epic Games, Inc. All Rights Reserved.
+
+#pragma once
+
+#include "IHierarchyTableColumn.h"
+
+#define LOCTEXT_NAMESPACE "HierarchyTableColumn_Mask"
+
+struct FHierarchyTableColumn_Mask : public IHierarchyTableColumn
+{
+public:
+	virtual FName GetColumnId() const override { return FName("Mask"); }
+	virtual FText GetColumnLabel() const override { return LOCTEXT("MaskLabel", "Mask"); }
+	virtual float GetColumnSize() const override { return 1.0f; }
+
+	virtual TSharedRef<SWidget> CreateEntryWidget(TObjectPtr<UHierarchyTable> HierarchyTable, int32 EntryIndex);
+	virtual TSharedRef<SWidget> CreateHeaderWidget();
+};
+
+#undef LOCTEXT_NAMESPACE
\ No newline at end of file
diff --git a/Engine/Plugins/Experimental/Animation/HierarchyTableAnimation/Source/Editor/Private/MaskProfile/MaskProfileFactory.cpp b/Engine/Plugins/Experimental/Animation/HierarchyTableAnimation/Source/Editor/Private/MaskProfile/MaskProfileFactory.cpp
deleted file mode 100644
diff --git a/Engine/Plugins/Experimental/Animation/HierarchyTableAnimation/Source/Editor/Private/MaskProfile/MaskProfileFactory.h b/Engine/Plugins/Experimental/Animation/HierarchyTableAnimation/Source/Editor/Private/MaskProfile/MaskProfileFactory.h
deleted file mode 100644
diff --git a/Engine/Plugins/Experimental/Animation/HierarchyTableAnimation/Source/Editor/Private/MaskProfile/MaskProfileWidgetConstructor.cpp b/Engine/Plugins/Experimental/Animation/HierarchyTableAnimation/Source/Editor/Private/MaskProfile/MaskProfileWidgetConstructor.cpp
deleted file mode 100644
diff --git a/Engine/Plugins/Experimental/Animation/HierarchyTableAnimation/Source/Editor/Private/MaskProfile/MaskProfileWidgetConstructor.h b/Engine/Plugins/Experimental/Animation/HierarchyTableAnimation/Source/Editor/Private/MaskProfile/MaskProfileWidgetConstructor.h
deleted file mode 100644
diff --git a/Engine/Plugins/Experimental/Animation/HierarchyTableAnimation/Source/Editor/Private/SkeletonHierarchyTableTypeHandler.cpp b/Engine/Plugins/Experimental/Animation/HierarchyTableAnimation/Source/Editor/Private/SkeletonHierarchyTableTypeHandler.cpp
--- a/Engine/Plugins/Experimental/Animation/HierarchyTableAnimation/Source/Editor/Private/SkeletonHierarchyTableTypeHandler.cpp
+++ b/Engine/Plugins/Experimental/Animation/HierarchyTableAnimation/Source/Editor/Private/SkeletonHierarchyTableTypeHandler.cpp
@@ -20,10 +20,154 @@
 #include "Widgets/Layout/SSpacer.h"
 #include "Widgets/Text/STextBlock.h"
 #include "Widgets/SBoxPanel.h"
-#include "IHierarchyTableEditor.h"
+#include "Widgets/Input/SEditableTextBox.h"
 
 #define LOCTEXT_NAMESPACE "UHierarchyTable_TableTypeHandler_Skeleton"
 
+FSlateIcon UHierarchyTable_TableTypeHandler_Skeleton::GetEntryIcon(const int32 EntryIndex) const
+{
+	const FHierarchyTableEntryData* const Entry = HierarchyTable->GetTableEntry(EntryIndex);
+	const FHierarchyTable_TablePayloadType_Skeleton& EntryTablePayload = Entry->TablePayload.Get<FHierarchyTable_TablePayloadType_Skeleton>();
+
+	switch (EntryTablePayload.EntryType)
+	{
+	case ESkeletonHierarchyTable_TablePayloadEntryType::Bone:
+		return FSlateIcon(FAppStyle::GetAppStyleSetName(), "SkeletonTree.Bone");
+	case ESkeletonHierarchyTable_TablePayloadEntryType::Curve:
+		return FSlateIcon(FAppStyle::GetAppStyleSetName(), "AnimGraph.Attribute.Curves.Icon");
+	case ESkeletonHierarchyTable_TablePayloadEntryType::Attribute:
+		return FSlateIcon(FAppStyle::GetAppStyleSetName(), "AnimGraph.Attribute.Attributes.Icon");
+	}
+
+	return FSlateIcon();
+}
+
+FSlateColor UHierarchyTable_TableTypeHandler_Skeleton::GetEntryIconColor(const int32 EntryIndex) const
+{
+	const FHierarchyTableEntryData* const Entry = HierarchyTable->GetTableEntry(EntryIndex);
+	const FHierarchyTable_TablePayloadType_Skeleton& EntryTablePayload = Entry->TablePayload.Get<FHierarchyTable_TablePayloadType_Skeleton>();
+
+	switch (EntryTablePayload.EntryType)
+	{
+	case ESkeletonHierarchyTable_TablePayloadEntryType::Bone:
+		return FSlateColor::UseForeground();
+	case ESkeletonHierarchyTable_TablePayloadEntryType::Curve:
+		return FAppStyle::GetSlateColor("AnimGraph.Attribute.Curves.Color");
+	case ESkeletonHierarchyTable_TablePayloadEntryType::Attribute:
+		return  FAppStyle::GetSlateColor("AnimGraph.Attribute.Attributes.Color");
+	}
+
+	return FSlateColor::UseForeground();
+}
+
+void UHierarchyTable_TableTypeHandler_Skeleton::ExtendContextMenu(FMenuBuilder& MenuBuilder, IHierarchyTable& HierarchyTableView) const
+{
+	const int32 SelectedIndex = HierarchyTableView.GetSelectedEntryIndex();
+
+	const FHierarchyTableEntryData* EntryData = HierarchyTable->GetTableEntry(SelectedIndex);
+	check(EntryData);
+
+	const FHierarchyTable_TablePayloadType_Skeleton& EntryMetadata = EntryData->GetMetadata<FHierarchyTable_TablePayloadType_Skeleton>();
+
+	FPersonaModule& PersonaModule = FModuleManager::LoadModuleChecked<FPersonaModule>("Persona");
+
+	const FHierarchyTable_TableType_Skeleton TableMetadata = HierarchyTable->GetTableMetadata<FHierarchyTable_TableType_Skeleton>();
+
+	MenuBuilder.AddSubMenu(
+		LOCTEXT("AddCurve_Label", "Add Curve"),
+		LOCTEXT("AddCurve_Tooltip", "Add a new curve entry parented to this entry"),
+		FNewMenuDelegate::CreateLambda([this, SelectedIndex, &HierarchyTableView, &PersonaModule, TableMetadata](FMenuBuilder& InSubMenuBuilder)
+			{
+				InSubMenuBuilder.BeginSection(NAME_None, LOCTEXT("NewCurveSection", "New Curve"));
+				{
+					InSubMenuBuilder.AddWidget(
+						SNew(SBox)
+						.Padding(FMargin(8.0f, 0.0f))
+						[
+							SNew(SEditableTextBox)
+								.Text(LOCTEXT("NewCurveDefault", "NewCurve"))
+								.OnTextCommitted_Lambda([this, SelectedIndex, &HierarchyTableView](const FText& InText, ETextCommit::Type InTextCommit)
+									{
+										if (InTextCommit == ETextCommit::OnEnter)
+										{
+											const int32 ParentIndex = SelectedIndex == INDEX_NONE ? 0 : SelectedIndex;
+											AddCurve(ParentIndex, FName(InText.ToString()));
+											HierarchyTableView.RegenerateHierarchy();
+										}
+									})
+								.OnVerifyTextChanged_Lambda([this](const FText& InNewText, FText& OutErrorMessage) -> bool
+									{
+										if (InNewText.IsEmpty())
+										{
+											OutErrorMessage = LOCTEXT("CurveNameEmpty", "Name can't be empty.");
+											return false;
+										}
+
+										if (HierarchyTable->HasIdentifier(FName(InNewText.ToString())))
+										{
+											OutErrorMessage = LOCTEXT("CurveNameExists", "Name already exists in the hierarchy.");
+											return false;
+										}
+
+										return true;
+									})
+						],
+						FText::GetEmpty(), true);
+				}
+				InSubMenuBuilder.EndSection();
+
+				InSubMenuBuilder.BeginSection(NAME_None, LOCTEXT("ExistingCurveSection", "Existing Curve"));
+				{
+					InSubMenuBuilder.AddWidget(
+						PersonaModule.CreateCurvePicker(TableMetadata.Skeleton,
+							FOnCurvePicked::CreateLambda([this, SelectedIndex, &HierarchyTableView](const FName& InName)
+								{
+									FSlateApplication::Get().DismissAllMenus();
+
+									const int32 ParentIndex = SelectedIndex == INDEX_NONE ? 0 : SelectedIndex;
+									AddCurve(ParentIndex, InName);
+									HierarchyTableView.RegenerateHierarchy();
+								}),
+							FIsCurveNameMarkedForExclusion::CreateLambda([this](const FName& InName)
+								{
+									return HierarchyTable->HasIdentifier(InName);
+								})),
+						FText::GetEmpty(), true);
+				}
+				InSubMenuBuilder.EndSection();
+			})
+	);
+
+	MenuBuilder.AddMenuEntry(
+		LOCTEXT("AddAttribute_Label", "Add Attribute"),
+		LOCTEXT("AddAttribute_Tooltip", "Add a new attribute entry parented to this entry"),
+		FSlateIcon(),
+		FUIAction(
+			FExecuteAction::CreateLambda([this, SelectedIndex, &HierarchyTableView]()
+				{
+					AddAttribute(SelectedIndex, "Test");
+					HierarchyTableView.RegenerateHierarchy();
+				}
+			)
+		));
+
+	if (EntryMetadata.EntryType == ESkeletonHierarchyTable_TablePayloadEntryType::Curve || EntryMetadata.EntryType == ESkeletonHierarchyTable_TablePayloadEntryType::Attribute)
+	{
+		MenuBuilder.AddMenuEntry(
+			LOCTEXT("RemoveEntry_Label", "Remove Entry"),
+			LOCTEXT("RemoveEntry_Tooltip", "Remove the selected entry"),
+			FSlateIcon(),
+			FUIAction(
+				FExecuteAction::CreateLambda([this, SelectedIndex, &HierarchyTableView]()
+					{
+						RemoveCurve(SelectedIndex);
+						HierarchyTableView.RegenerateHierarchy();
+					}
+				)
+			));
+	}
+}
+
 void UHierarchyTable_TableTypeHandler_Skeleton::ConstructHierarchy()
 {
 	const FHierarchyTable_TableType_Skeleton& SkeletonTableType = HierarchyTable->GetTableMetadata<FHierarchyTable_TableType_Skeleton>();
@@ -87,165 +231,8 @@ bool UHierarchyTable_TableTypeHandler_Skeleton::FactoryConfigureProperties(FInst
 	return SkeletonTableType.Skeleton != nullptr;
 }
 
-void UHierarchyTable_TableTypeHandler_Skeleton::ExtendToolbar(UToolMenu* ToolMenu, const IHierarchyTableEditor* Editor) const
-{
-	const FToolMenuInsert SectionInsertLocation("Asset", EToolMenuInsertType::After);
-	{
-		FToolMenuSection& HierarchyTableSection = ToolMenu->AddSection("HierarchyTable", LOCTEXT("HierarchyTable_ToolbarLabel", "HierarchyTable"), SectionInsertLocation);
-
-		HierarchyTableSection.AddEntry(FToolMenuEntry::InitComboButton(
-			"AddCurve",
-			FUIAction(),
-			FNewToolMenuDelegate::CreateLambda([=, this](UToolMenu* InSubMenu)
-				{
-					FToolUIAction Action;
-				    Action.ExecuteAction = FToolMenuExecuteAction::CreateLambda([=, this](const FToolMenuContext& Context)
-				    {
-						TSharedRef<STextEntryPopup> TextEntry = SNew(STextEntryPopup)
-							.Label(LOCTEXT("NewCurveEntryLabal", "Curve Name"))
-							.OnTextCommitted_Lambda([=, this](const FText& CommittedText, ETextCommit::Type CommitInfo)
-								{
-									const int32 SelectedIndex = Editor->GetSelectedEntryIndex();
-									const int32 ParentIndex = SelectedIndex == INDEX_NONE ? 0 : SelectedIndex;
-									AddCurve(ParentIndex, FName(CommittedText.ToString()));
-									FSlateApplication::Get().DismissAllMenus();
-								});
-
-						FSlateApplication& SlateApp = FSlateApplication::Get();
-						SlateApp.PushMenu(
-							SlateApp.GetInteractiveTopLevelWindows()[0],
-							FWidgetPath(),
-							TextEntry,
-							SlateApp.GetCursorPos(),
-							FPopupTransitionEffect::TypeInPopup);
-				    });
-
-				    FToolMenuEntry Entry = FToolMenuEntry::InitMenuEntry(
-					    FName("AddNewCurve"),
-					    LOCTEXT("AddNewCurve_Label", "Add New Curve"),
-					    LOCTEXT("AddNewCurve_Tooltip", "Add a new curve value"),
-					    FSlateIcon(),
-					    FToolUIActionChoice(Action),
-						EUserInterfaceActionType::Button);
-				
-				    InSubMenu->AddMenuEntry("AddNewCurve", Entry);
-
-					const FHierarchyTable_TableType_Skeleton& SkeletonTableMetadata = HierarchyTable->GetTableMetadata<FHierarchyTable_TableType_Skeleton>();
-
-					FPersonaModule& PersonaModule = FModuleManager::LoadModuleChecked<FPersonaModule>("Persona");
-					InSubMenu->AddMenuEntry("ExistingCurveMenu", FToolMenuEntry::InitWidget(
-						"ExistingCurveMenu",
-						SNew(SVerticalBox)
-						+SVerticalBox::Slot()
-						.AutoHeight()
-						[
-							PersonaModule.CreateCurvePicker(SkeletonTableMetadata.Skeleton,
-								FOnCurvePicked::CreateLambda([&](const FName& InName)
-									{
-										const int32 SelectedIndex = Editor->GetSelectedEntryIndex();
-										const int32 ParentIndex = SelectedIndex == INDEX_NONE ? 0 : SelectedIndex;
-										AddCurve(ParentIndex, InName);
-										FSlateApplication::Get().DismissAllMenus();
-									}))
-						],
-						FText(),
-						true,
-						false,
-						true
-					));
-				}),
-			LOCTEXT("AddCurve_Label", "Add Curve"),
-			LOCTEXT("AddCurve_ToolTip", "Add a new curve to the hierarchy"),
-			FSlateIcon(FAppStyle::Get().GetStyleSetName(), "Icons.Plus")
-		));
-
-		HierarchyTableSection.AddEntry(FToolMenuEntry::InitComboButton(
-			"AddAttribute",
-			FUIAction(),
-			FNewToolMenuDelegate::CreateLambda([=, this](UToolMenu* InSubMenu)
-				{
-					FToolUIAction Action;
-				    Action.ExecuteAction = FToolMenuExecuteAction::CreateLambda([=, this](const FToolMenuContext& Context)
-				    {
-						TSharedRef<STextEntryPopup> TextEntry = SNew(STextEntryPopup)
-							.Label(LOCTEXT("NewAttributeEntryLabal", "Attribute Name"))
-							.OnTextCommitted_Lambda([=, this](const FText& CommittedText, ETextCommit::Type CommitInfo)
-								{
-									const int32 SelectedIndex = Editor->GetSelectedEntryIndex();
-									const int32 ParentIndex = SelectedIndex == INDEX_NONE ? 0 : SelectedIndex;
-									AddAttribute(ParentIndex, FName(CommittedText.ToString()));
-									FSlateApplication::Get().DismissAllMenus();
-								});
-
-						FSlateApplication& SlateApp = FSlateApplication::Get();
-						SlateApp.PushMenu(
-							SlateApp.GetInteractiveTopLevelWindows()[0],
-							FWidgetPath(),
-							TextEntry,
-							SlateApp.GetCursorPos(),
-							FPopupTransitionEffect::TypeInPopup);
-				    });
-
-				    FToolMenuEntry Entry = FToolMenuEntry::InitMenuEntry(
-					    FName("AddNewAttribute"),
-					    LOCTEXT("AddNewAttribute_Label", "Add New Attribute"),
-					    LOCTEXT("AddNewAttribute_Tooltip", "Add a new attribute value"),
-					    FSlateIcon(),
-					    FToolUIActionChoice(Action),
-						EUserInterfaceActionType::Button);
-				
-				    InSubMenu->AddMenuEntry("AddNewAttribute", Entry);
-				}),
-			LOCTEXT("AddAttribute_Label", "Add Attribute"),
-			LOCTEXT("AddAttribute_ToolTip", "Add a new attribute to the hierarchy"),
-			FSlateIcon(FAppStyle::Get().GetStyleSetName(), "Icons.Plus")
-		));
-	}
-}
-
-TSharedRef<SWidget> UHierarchyTable_TableTypeHandler_Skeleton::CreateLabelWidget(const int32 EntryIndex) const
+void UHierarchyTable_TableTypeHandler_Skeleton::ExtendToolbar(UToolMenu* ToolMenu, IHierarchyTable& HierarchyTableView) const
 {
-	const FSlateBrush* LabelIcon = nullptr;
-	FSlateColor LabelIconColor = FSlateColor::UseForeground();
-
-	const FHierarchyTableEntryData* const Entry = HierarchyTable->GetTableEntry(EntryIndex);
-	const FHierarchyTable_TablePayloadType_Skeleton& EntryTablePayload = Entry->TablePayload.Get<FHierarchyTable_TablePayloadType_Skeleton>();
-
-	switch (EntryTablePayload.EntryType)
-	{
-		case ESkeletonHierarchyTable_TablePayloadEntryType::Bone:
-			LabelIcon = FAppStyle::GetBrush("SkeletonTree.Bone");
-			break;
-		case ESkeletonHierarchyTable_TablePayloadEntryType::Curve:
-			LabelIcon = FAppStyle::GetBrush("AnimGraph.Attribute.Curves.Icon");
-			LabelIconColor = FAppStyle::GetSlateColor("AnimGraph.Attribute.Curves.Color");
-			break;
-		case ESkeletonHierarchyTable_TablePayloadEntryType::Attribute:
-			LabelIcon = FAppStyle::GetBrush("AnimGraph.Attribute.Attributes.Icon");
-			LabelIconColor = FAppStyle::GetSlateColor("AnimGraph.Attribute.Attributes.Color");
-			break;
-	}
-
-	return SNew(SHorizontalBox)
-		+SHorizontalBox::Slot()
-		.AutoWidth()
-		[
-			SNew(SImage)
-			.Image(LabelIcon)
-			.ColorAndOpacity(LabelIconColor)
-		]
-		+SHorizontalBox::Slot()
-		.AutoWidth()
-		[
-			SNew(SSpacer)
-			.Size(FVector2D(5.0f, 0.0f))
-		]
-		+SHorizontalBox::Slot()
-		.FillWidth(1.0f)
-		[
-			SNew(STextBlock)
-			.Text(FText::FromName(Entry->Identifier))
-		];
 }
 
 void UHierarchyTable_TableTypeHandler_Skeleton::AddCurve(const int32 ParentIndex, const FName Identifier) const
@@ -266,6 +253,11 @@ void UHierarchyTable_TableTypeHandler_Skeleton::AddCurve(const int32 ParentIndex
 
 	HierarchyTable->AddEntry(NewEntry);
 }
+
+void UHierarchyTable_TableTypeHandler_Skeleton::RemoveCurve(const int32 Index) const
+{
+	HierarchyTable->RemoveEntry(Index);
+}
 	
 void UHierarchyTable_TableTypeHandler_Skeleton::AddAttribute(const int32 ParentIndex, const FName Identifier) const
 {
diff --git a/Engine/Plugins/Experimental/Animation/HierarchyTableAnimation/Source/Editor/Private/SkeletonHierarchyTableTypeHandler.h b/Engine/Plugins/Experimental/Animation/HierarchyTableAnimation/Source/Editor/Private/SkeletonHierarchyTableTypeHandler.h
--- a/Engine/Plugins/Experimental/Animation/HierarchyTableAnimation/Source/Editor/Private/SkeletonHierarchyTableTypeHandler.h
+++ b/Engine/Plugins/Experimental/Animation/HierarchyTableAnimation/Source/Editor/Private/SkeletonHierarchyTableTypeHandler.h
@@ -16,12 +16,17 @@ public:
 
 	virtual bool FactoryConfigureProperties(FInstancedStruct& TableType) const override;
 
-	virtual void ExtendToolbar(UToolMenu* ToolMenu, const IHierarchyTableEditor* Editor) const override;
+	virtual void ExtendToolbar(UToolMenu* ToolMenu, IHierarchyTable& HierarchyTableView) const override;
 
-	virtual TSharedRef<SWidget> CreateLabelWidget(const int32 EntryIndex) const;
+	virtual void ExtendContextMenu(FMenuBuilder& MenuBuilder, IHierarchyTable& HierarchyTableView) const override;
+
+	virtual FSlateIcon GetEntryIcon(const int32 EntryIndex) const override;
+
+	virtual FSlateColor GetEntryIconColor(const int32 EntryIndex) const override;
 
 private:
 	void AddCurve(const int32 ParentIndex, const FName Identifier) const;
+	void RemoveCurve(const int32 Index) const;
 	
 	void AddAttribute(const int32 ParentIndex, const FName Identifier) const;
 };
\ No newline at end of file
diff --git a/Engine/Plugins/Experimental/Animation/HierarchyTableAnimation/Source/Editor/Private/TimeProfile/TimeProfileColumn.cpp b/Engine/Plugins/Experimental/Animation/HierarchyTableAnimation/Source/Editor/Private/TimeProfile/TimeProfileColumn.cpp
new file mode 100644
--- /dev/null
+++ b/Engine/Plugins/Experimental/Animation/HierarchyTableAnimation/Source/Editor/Private/TimeProfile/TimeProfileColumn.cpp
@@ -0,0 +1,83 @@
+// Copyright Epic Games, Inc. All Rights Reserved.
+
+#include "TimeProfileColumn.h"
+#include "Widgets/Input/SSpinBox.h"
+#include "TimeProfile/HierarchyTableTypeTime.h"
+#include "HierarchyTable.h"
+
+TSharedRef<SWidget> FHierarchyTableColumn_TimeProfile_StartTime::CreateEntryWidget(TObjectPtr<UHierarchyTable> HierarchyTable, int32 EntryIndex)
+{
+	return SNew(SSpinBox<float>)
+		.IsEnabled_Lambda([HierarchyTable, EntryIndex]() { return HierarchyTable->GetTableEntry(EntryIndex)->IsOverridden(); })
+		.MinDesiredWidth(100.0f)
+		.MinValue(0.0f)
+		.MaxValue(1.0f)
+		.Value_Lambda([HierarchyTable, EntryIndex]()
+			{
+				return HierarchyTable->GetTableEntry(EntryIndex)->GetValue<FHierarchyTable_ElementType_Time>()->StartTime;
+			})
+		.OnValueChanged_Lambda([HierarchyTable, EntryIndex](float NewValue)
+			{
+				HierarchyTable->GetMutableTableEntry(EntryIndex)->GetMutableValue<FHierarchyTable_ElementType_Time>()->StartTime = NewValue;
+			});
+}
+
+TSharedRef<SWidget> FHierarchyTableColumn_TimeProfile_StartTime::CreateHeaderWidget()
+{
+	return SNullWidget::NullWidget;
+}
+
+TSharedRef<SWidget> FHierarchyTableColumn_TimeProfile_EndTime::CreateEntryWidget(TObjectPtr<UHierarchyTable> HierarchyTable, int32 EntryIndex)
+{
+	return SNew(SSpinBox<float>)
+		.IsEnabled_Lambda([HierarchyTable, EntryIndex]() { return HierarchyTable->GetTableEntry(EntryIndex)->IsOverridden(); })
+		.MinDesiredWidth(100.0f)
+		.MinValue(0.0f)
+		.MaxValue(1.0f)
+		.Value_Lambda([HierarchyTable, EntryIndex]()
+			{
+				return HierarchyTable->GetTableEntry(EntryIndex)->GetValue<FHierarchyTable_ElementType_Time>()->EndTime;
+			})
+		.OnValueChanged_Lambda([HierarchyTable, EntryIndex](float NewValue)
+			{
+				HierarchyTable->GetMutableTableEntry(EntryIndex)->GetMutableValue<FHierarchyTable_ElementType_Time>()->EndTime = NewValue;
+			});
+}
+
+TSharedRef<SWidget> FHierarchyTableColumn_TimeProfile_EndTime::CreateHeaderWidget()
+{
+	return SNullWidget::NullWidget;
+}
+
+TSharedRef<SWidget> FHierarchyTableColumn_TimeProfile_TimeFactor::CreateEntryWidget(TObjectPtr<UHierarchyTable> HierarchyTable, int32 EntryIndex)
+{
+	return SNew(SSpinBox<float>)
+		.IsEnabled_Lambda([HierarchyTable, EntryIndex]() { return HierarchyTable->GetTableEntry(EntryIndex)->IsOverridden(); })
+		.MinDesiredWidth(100.0f)
+		.MinValue(0.0f)
+		.MaxValue(1.0f)
+		.Value_Lambda([HierarchyTable, EntryIndex]()
+			{
+				return HierarchyTable->GetTableEntry(EntryIndex)->GetValue<FHierarchyTable_ElementType_Time>()->TimeFactor;
+			})
+		.OnValueChanged_Lambda([HierarchyTable, EntryIndex](float NewValue)
+			{
+				HierarchyTable->GetMutableTableEntry(EntryIndex)->GetMutableValue<FHierarchyTable_ElementType_Time>()->TimeFactor = NewValue;
+			});
+}
+
+TSharedRef<SWidget> FHierarchyTableColumn_TimeProfile_TimeFactor::CreateHeaderWidget()
+{
+	return SNullWidget::NullWidget;
+}
+
+TSharedRef<SWidget> FHierarchyTableColumn_TimeProfile_Preview::CreateEntryWidget(TObjectPtr<UHierarchyTable> HierarchyTable, int32 EntryIndex)
+{
+	return SNew(STextBlock)
+		.Text(INVTEXT("PREVIEW"));
+}
+
+TSharedRef<SWidget> FHierarchyTableColumn_TimeProfile_Preview::CreateHeaderWidget()
+{
+	return SNullWidget::NullWidget;
+}
\ No newline at end of file
diff --git a/Engine/Plugins/Experimental/Animation/HierarchyTableAnimation/Source/Editor/Private/TimeProfile/TimeProfileColumn.h b/Engine/Plugins/Experimental/Animation/HierarchyTableAnimation/Source/Editor/Private/TimeProfile/TimeProfileColumn.h
new file mode 100644
--- /dev/null
+++ b/Engine/Plugins/Experimental/Animation/HierarchyTableAnimation/Source/Editor/Private/TimeProfile/TimeProfileColumn.h
@@ -0,0 +1,53 @@
+// Copyright Epic Games, Inc. All Rights Reserved.
+
+#pragma once
+
+#include "IHierarchyTableColumn.h"
+
+#define LOCTEXT_NAMESPACE "HierarchyTableColumn_TimeProfile"
+
+struct FHierarchyTableColumn_TimeProfile_StartTime : public IHierarchyTableColumn
+{
+public:
+	virtual FName GetColumnId() const override { return FName("TimeProfile_StartTime"); }
+	virtual FText GetColumnLabel() const override { return LOCTEXT("StartTimeLabel", "Start Time"); }
+	virtual float GetColumnSize() const override { return 0.25f; }
+
+	virtual TSharedRef<SWidget> CreateEntryWidget(TObjectPtr<UHierarchyTable> HierarchyTable, int32 EntryIndex);
+	virtual TSharedRef<SWidget> CreateHeaderWidget();
+};
+
+struct FHierarchyTableColumn_TimeProfile_EndTime : public IHierarchyTableColumn
+{
+public:
+	virtual FName GetColumnId() const override { return FName("TimeProfile_EndTime"); }
+	virtual FText GetColumnLabel() const override { return LOCTEXT("EndTimeLabel", "End Time"); }
+	virtual float GetColumnSize() const override { return 0.25f; }
+
+	virtual TSharedRef<SWidget> CreateEntryWidget(TObjectPtr<UHierarchyTable> HierarchyTable, int32 EntryIndex);
+	virtual TSharedRef<SWidget> CreateHeaderWidget();
+};
+
+struct FHierarchyTableColumn_TimeProfile_TimeFactor : public IHierarchyTableColumn
+{
+public:
+	virtual FName GetColumnId() const override { return FName("TimeProfile_TimeFactor"); }
+	virtual FText GetColumnLabel() const override { return LOCTEXT("TimeFactorLabel", "Time Factor"); }
+	virtual float GetColumnSize() const override { return 0.25f; }
+
+	virtual TSharedRef<SWidget> CreateEntryWidget(TObjectPtr<UHierarchyTable> HierarchyTable, int32 EntryIndex);
+	virtual TSharedRef<SWidget> CreateHeaderWidget();
+};
+
+struct FHierarchyTableColumn_TimeProfile_Preview : public IHierarchyTableColumn
+{
+public:
+	virtual FName GetColumnId() const override { return FName("TimeProfile_Preview"); }
+	virtual FText GetColumnLabel() const override { return LOCTEXT("StartTimeLabel", "Preview"); }
+	virtual float GetColumnSize() const override { return 0.25f; }
+
+	virtual TSharedRef<SWidget> CreateEntryWidget(TObjectPtr<UHierarchyTable> HierarchyTable, int32 EntryIndex);
+	virtual TSharedRef<SWidget> CreateHeaderWidget();
+};
+
+#undef LOCTEXT_NAMESPACE
\ No newline at end of file
diff --git a/Engine/Plugins/Experimental/Animation/HierarchyTableAnimation/Source/Editor/Private/TimeProfile/TimeProfileFactory.cpp b/Engine/Plugins/Experimental/Animation/HierarchyTableAnimation/Source/Editor/Private/TimeProfile/TimeProfileFactory.cpp
deleted file mode 100644
diff --git a/Engine/Plugins/Experimental/Animation/HierarchyTableAnimation/Source/Editor/Private/TimeProfile/TimeProfileFactory.h b/Engine/Plugins/Experimental/Animation/HierarchyTableAnimation/Source/Editor/Private/TimeProfile/TimeProfileFactory.h
deleted file mode 100644
diff --git a/Engine/Plugins/Experimental/Animation/HierarchyTableAnimation/Source/Editor/Private/TimeProfile/TimeProfileWidgetConstructor.cpp b/Engine/Plugins/Experimental/Animation/HierarchyTableAnimation/Source/Editor/Private/TimeProfile/TimeProfileWidgetConstructor.cpp
deleted file mode 100644
diff --git a/Engine/Plugins/Experimental/Animation/HierarchyTableAnimation/Source/Editor/Private/TimeProfile/TimeProfileWidgetConstructor.h b/Engine/Plugins/Experimental/Animation/HierarchyTableAnimation/Source/Editor/Private/TimeProfile/TimeProfileWidgetConstructor.h
deleted file mode 100644
diff --git a/Engine/Plugins/Experimental/Animation/HierarchyTableAnimation/Source/Editor/Public/MaskProfile/MaskProfileProxyColumn.h b/Engine/Plugins/Experimental/Animation/HierarchyTableAnimation/Source/Editor/Public/MaskProfile/MaskProfileProxyColumn.h
deleted file mode 100644
diff --git a/Engine/Plugins/Experimental/Animation/HierarchyTableAnimation/Source/Editor/Public/TimeProfile/TimeProfileProxyColumn.h b/Engine/Plugins/Experimental/Animation/HierarchyTableAnimation/Source/Editor/Public/TimeProfile/TimeProfileProxyColumn.h
deleted file mode 100644
