Define ShortNames for Hierarchy Table modules

diff --git a/Engine/Plugins/Experimental/Animation/HierarchyTable/Source/Editor/HierarchyTableEditor.Build.cs b/Engine/Plugins/Experimental/Animation/HierarchyTable/Source/Editor/HierarchyTableEditor.Build.cs
--- a/Engine/Plugins/Experimental/Animation/HierarchyTable/Source/Editor/HierarchyTableEditor.Build.cs
+++ b/Engine/Plugins/Experimental/Animation/HierarchyTable/Source/Editor/HierarchyTableEditor.Build.cs
@@ -7,6 +7,8 @@ public class HierarchyTableEditor : ModuleRules
 {
 	public HierarchyTableEditor(ReadOnlyTargetRules Target) : base(Target)
 	{
+		ShortName = "HTEd";
+
 		PrivateDependencyModuleNames.AddAll(
 			"AssetDefinition",
 			"Core",
diff --git a/Engine/Plugins/Experimental/Animation/HierarchyTable/Source/Runtime/HierarchyTableRuntime.Build.cs b/Engine/Plugins/Experimental/Animation/HierarchyTable/Source/Runtime/HierarchyTableRuntime.Build.cs
--- a/Engine/Plugins/Experimental/Animation/HierarchyTable/Source/Runtime/HierarchyTableRuntime.Build.cs
+++ b/Engine/Plugins/Experimental/Animation/HierarchyTable/Source/Runtime/HierarchyTableRuntime.Build.cs
@@ -6,6 +6,8 @@ public class HierarchyTableRuntime : ModuleRules
 {
 	public HierarchyTableRuntime(ReadOnlyTargetRules Target) : base(Target)
 	{
+		ShortName = "HTRun";
+
 		PublicDependencyModuleNames.AddRange(
 			new string[]
 			{
diff --git a/Engine/Plugins/Experimental/Animation/HierarchyTableAnimation/Source/Editor/HierarchyTableAnimationEditor.Build.cs b/Engine/Plugins/Experimental/Animation/HierarchyTableAnimation/Source/Editor/HierarchyTableAnimationEditor.Build.cs
--- a/Engine/Plugins/Experimental/Animation/HierarchyTableAnimation/Source/Editor/HierarchyTableAnimationEditor.Build.cs
+++ b/Engine/Plugins/Experimental/Animation/HierarchyTableAnimation/Source/Editor/HierarchyTableAnimationEditor.Build.cs
@@ -7,6 +7,8 @@ public class HierarchyTableAnimationEditor : ModuleRules
 {
 	public HierarchyTableAnimationEditor(ReadOnlyTargetRules Target) : base(Target)
 	{
+		ShortName = "HTAnimEd";
+
 		PrivateDependencyModuleNames.AddAll(
 			"Core",
 			"CoreUObject",
diff --git a/Engine/Plugins/Experimental/Animation/HierarchyTableAnimation/Source/Runtime/HierarchyTableAnimationRuntime.Build.cs b/Engine/Plugins/Experimental/Animation/HierarchyTableAnimation/Source/Runtime/HierarchyTableAnimationRuntime.Build.cs
--- a/Engine/Plugins/Experimental/Animation/HierarchyTableAnimation/Source/Runtime/HierarchyTableAnimationRuntime.Build.cs
+++ b/Engine/Plugins/Experimental/Animation/HierarchyTableAnimation/Source/Runtime/HierarchyTableAnimationRuntime.Build.cs
@@ -6,6 +6,8 @@ public class HierarchyTableAnimationRuntime : ModuleRules
 {
 	public HierarchyTableAnimationRuntime(ReadOnlyTargetRules Target) : base(Target)
 	{
+		ShortName = "HTAnimRun";
+
 		PublicDependencyModuleNames.AddRange(
 			new string[]
 			{
diff --git a/Engine/Plugins/Experimental/Animation/HierarchyTableAnimation/Source/UncookedOnly/HierarchyTableAnimationUncookedOnly.Build.cs b/Engine/Plugins/Experimental/Animation/HierarchyTableAnimation/Source/UncookedOnly/HierarchyTableAnimationUncookedOnly.Build.cs
--- a/Engine/Plugins/Experimental/Animation/HierarchyTableAnimation/Source/UncookedOnly/HierarchyTableAnimationUncookedOnly.Build.cs
+++ b/Engine/Plugins/Experimental/Animation/HierarchyTableAnimation/Source/UncookedOnly/HierarchyTableAnimationUncookedOnly.Build.cs
@@ -7,6 +7,8 @@ public class HierarchyTableAnimationUncookedOnly : ModuleRules
 {
 	public HierarchyTableAnimationUncookedOnly(ReadOnlyTargetRules Target) : base(Target)
 	{
+		ShortName = "HTAnimUncook";
+
 		PublicDependencyModuleNames.AddRange(
 			new string[]
 			{
