Fix LOCTEXT conflict in hierarchy table animation

diff --git a/Engine/Plugins/Experimental/Animation/HierarchyTableAnimation/Source/Editor/Private/TimeProfile/TimeProfileColumn.h b/Engine/Plugins/Experimental/Animation/HierarchyTableAnimation/Source/Editor/Private/TimeProfile/TimeProfileColumn.h
--- a/Engine/Plugins/Experimental/Animation/HierarchyTableAnimation/Source/Editor/Private/TimeProfile/TimeProfileColumn.h
+++ b/Engine/Plugins/Experimental/Animation/HierarchyTableAnimation/Source/Editor/Private/TimeProfile/TimeProfileColumn.h
@@ -43,7 +43,7 @@ struct FHierarchyTableColumn_TimeProfile_Preview : public IHierarchyTableColumn
 {
 public:
 	virtual FName GetColumnId() const override { return FName("TimeProfile_Preview"); }
-	virtual FText GetColumnLabel() const override { return LOCTEXT("StartTimeLabel", "Preview"); }
+	virtual FText GetColumnLabel() const override { return LOCTEXT("PreviewLabel", "Preview"); }
 	virtual float GetColumnSize() const override { return 0.25f; }
 
 	virtual TSharedRef<SWidget> CreateEntryWidget(TObjectPtr<UHierarchyTable> HierarchyTable, int32 EntryIndex);
