Hierarchy Tables - Refactor removing/renaming entries

diff --git a/Engine/Plugins/Experimental/Animation/HierarchyTable/Source/Editor/Private/HierarchyTableTypeHandler.cpp b/Engine/Plugins/Experimental/Animation/HierarchyTable/Source/Editor/Private/HierarchyTableTypeHandler.cpp
--- a/Engine/Plugins/Experimental/Animation/HierarchyTable/Source/Editor/Private/HierarchyTableTypeHandler.cpp
+++ b/Engine/Plugins/Experimental/Animation/HierarchyTable/Source/Editor/Private/HierarchyTableTypeHandler.cpp
@@ -17,4 +17,36 @@ FSlateIcon UHierarchyTable_TableTypeHandler::GetEntryIcon(const int32 EntryIndex
 FSlateColor UHierarchyTable_TableTypeHandler::GetEntryIconColor(const int32 EntryIndex) const
 {
 	return FSlateColor::UseForeground();
+}
+
+bool UHierarchyTable_TableTypeHandler::RenameEntry(const int32 EntryIndex, const FName NewName)
+{
+	if (!ensure(CanRenameEntry(EntryIndex)))
+	{
+		return false;
+	}
+
+	if (HierarchyTable->HasIdentifier(NewName))
+	{
+		return false;
+	}
+
+	if (FHierarchyTableEntryData* const Entry = HierarchyTable->GetMutableTableEntry(EntryIndex))
+	{
+		Entry->Identifier = NewName;
+	}
+
+	return true;
+}
+
+bool UHierarchyTable_TableTypeHandler::RemoveEntry(const int32 EntryIndex)
+{
+	if (!ensure(CanRemoveEntry(EntryIndex)))
+	{
+		return false;
+	}
+
+	HierarchyTable->RemoveEntry(EntryIndex);
+
+	return true;
 }
\ No newline at end of file
diff --git a/Engine/Plugins/Experimental/Animation/HierarchyTable/Source/Editor/Private/SHierarchyTable.cpp b/Engine/Plugins/Experimental/Animation/HierarchyTable/Source/Editor/Private/SHierarchyTable.cpp
--- a/Engine/Plugins/Experimental/Animation/HierarchyTable/Source/Editor/Private/SHierarchyTable.cpp
+++ b/Engine/Plugins/Experimental/Animation/HierarchyTable/Source/Editor/Private/SHierarchyTable.cpp
@@ -3,11 +3,13 @@
 #include "SHierarchyTable.h"
 
 #include "HierarchyTableEditorModule.h"
-#include "Modules/ModuleManager.h"
+#include "Framework/Commands/GenericCommands.h"
 #include "HierarchyTable.h"
-#include "Widgets/Views/STreeView.h"
-#include "SHierarchyTableRow.h"
 #include "IHierarchyTableColumn.h"
+#include "Modules/ModuleManager.h"
+#include "ScopedTransaction.h"
+#include "SHierarchyTableRow.h"
+#include "Widgets/Views/STreeView.h"
 
 #define LOCTEXT_NAMESPACE "SHierarchyTable"
 
@@ -110,7 +112,11 @@ void SHierarchyTable::TreeView_OnItemScrolledIntoView(TSharedPtr<FTreeItem> InIt
 
 TSharedRef<ITableRow> SHierarchyTable::TreeView_GenerateItemRow(TSharedPtr<SHierarchyTable::FTreeItem> Item, const TSharedRef<STableViewBase>& OwnerTable)
 {
-	return SNew(SHierarchyTableRow, OwnerTable, SharedThis(this), Item);
+	return SNew(SHierarchyTableRow, OwnerTable, SharedThis(this), Item)
+		.OnRenamed_Lambda([this, Item](const FName NewName)
+			{
+				return TableHandler->RenameEntry(Item->Index, NewName);
+			});
 }
 
 TSharedPtr<SWidget> SHierarchyTable::TreeView_HandleContextMenuOpening()
@@ -123,6 +129,52 @@ TSharedPtr<SWidget> SHierarchyTable::TreeView_HandleContextMenuOpening()
 	if (SelectedItems.Num() > 0)
 	{
 		TableHandler->ExtendContextMenu(MenuBuilder, *this);
+
+		MenuBuilder.AddMenuSeparator();
+
+		MenuBuilder.AddMenuEntry(
+			LOCTEXT("RenameEntry_Label", "Rename"),
+			LOCTEXT("RenameEntry_Tooltip", "Rename the selected entry"),
+			FGenericCommands::Get().Rename->GetIcon(),
+			FUIAction(
+				FExecuteAction::CreateLambda([this, SelectedItems]()
+					{
+						const FScopedTransaction Transaction(LOCTEXT("RenameEntry_Transaction", "Rename Entry"));
+						HierarchyTable->Modify();
+
+						SelectedItems[0]->OnRenameRequested.ExecuteIfBound();
+					}
+				),
+				FCanExecuteAction::CreateLambda([this, SelectedItems]()
+					{
+						return TableHandler->CanRenameEntry(SelectedItems[0]->Index);
+					}
+				)
+			));
+
+		MenuBuilder.AddMenuEntry(
+			LOCTEXT("RemoveEntry_Label", "Remove"),
+			LOCTEXT("RemoveEntry_Tooltip", "Remove the selected entry"),
+			FGenericCommands::Get().Delete->GetIcon(),
+			FUIAction(
+				FExecuteAction::CreateLambda([this, SelectedItems]()
+					{
+						const FScopedTransaction Transaction(LOCTEXT("RemoveEntry_Transaction", "Remove Entry"));
+						HierarchyTable->Modify();
+
+						const bool bSuccess = TableHandler->RemoveEntry(SelectedItems[0]->Index);
+						if (ensure(bSuccess))
+						{
+							RegenerateTreeViewItems();
+						}
+					}
+				),
+				FCanExecuteAction::CreateLambda([this, SelectedItems]()
+					{
+						return TableHandler->CanRemoveEntry(SelectedItems[0]->Index);
+					}
+				)
+			));
 	}
 
 	return MenuBuilder.MakeWidget();
@@ -205,4 +257,14 @@ void SHierarchyTable::RegenerateTreeViewItems()
 	}
 }
 
+void SHierarchyTable::PostUndo(bool bSuccess)
+{
+	RegenerateTreeViewItems();
+}
+
+void SHierarchyTable::PostRedo(bool bSuccess)
+{
+	RegenerateTreeViewItems();
+}
+
 #undef LOCTEXT_NAMESPACE
\ No newline at end of file
diff --git a/Engine/Plugins/Experimental/Animation/HierarchyTable/Source/Editor/Private/SHierarchyTable.h b/Engine/Plugins/Experimental/Animation/HierarchyTable/Source/Editor/Private/SHierarchyTable.h
--- a/Engine/Plugins/Experimental/Animation/HierarchyTable/Source/Editor/Private/SHierarchyTable.h
+++ b/Engine/Plugins/Experimental/Animation/HierarchyTable/Source/Editor/Private/SHierarchyTable.h
@@ -7,13 +7,14 @@
 #include "Widgets/SCompoundWidget.h"
 #include "Widgets/Views/STableViewBase.h"
 #include "Widgets/Views/STreeView.h"
+#include "EditorUndoClient.h"
 
 class ITableRow;
 class SHierarchyTableRow;
 class UHierarchyTable;
 class UHierarchyTable_TableTypeHandler;
 
-class SHierarchyTable : public SCompoundWidget, public IHierarchyTable
+class SHierarchyTable : public SCompoundWidget, public IHierarchyTable, public FSelfRegisteringEditorUndoClient
 {
 public:
 	friend SHierarchyTableRow;
@@ -43,6 +44,10 @@ public:
 	int32 GetSelectedEntryIndex() const override;
 	void RegenerateHierarchy() override;
 
+	// FSelfRegisteringEditorUndoClient
+	void PostUndo(bool bSuccess) override;
+	void PostRedo(bool bSuccess) override;
+
 private:
 	void TreeView_HandleGetChildren(TSharedPtr<FTreeItem> InItem, TArray<TSharedPtr<FTreeItem>>& OutChildren) const;
 
diff --git a/Engine/Plugins/Experimental/Animation/HierarchyTable/Source/Editor/Private/SHierarchyTableRow.cpp b/Engine/Plugins/Experimental/Animation/HierarchyTable/Source/Editor/Private/SHierarchyTableRow.cpp
--- a/Engine/Plugins/Experimental/Animation/HierarchyTable/Source/Editor/Private/SHierarchyTableRow.cpp
+++ b/Engine/Plugins/Experimental/Animation/HierarchyTable/Source/Editor/Private/SHierarchyTableRow.cpp
@@ -52,7 +52,7 @@ void SHierarchyTableRow::OnCommitRename(const FText& InText, ETextCommit::Type C
 
 	if (OnRenamed.IsBound())
 	{
-		const bool bSuccess = OnRenamed.Execute(OldName, NewName);
+		const bool bSuccess = OnRenamed.Execute(NewName);
 		if (bSuccess)
 		{
 			TreeItem->Name = NewName;
@@ -89,6 +89,24 @@ TSharedRef<SWidget> SHierarchyTableRow::GenerateWidgetForColumn(const FName& Col
 						{
 							return FText::FromName(TreeItem->Name);
 						})
+					.OnVerifyTextChanged_Lambda([this](const FText& InNewText, FText& OutErrorMessage) -> bool
+						{
+							if (InNewText.IsEmpty())
+							{
+								OutErrorMessage = LOCTEXT("AttributeNameEmpty", "Name can't be empty.");
+								return false;
+							}
+
+							const FName CurrentName = HierarchyTableWidget->HierarchyTable->GetTableEntry(TreeItem->Index)->Identifier;
+
+							if (HierarchyTableWidget->HierarchyTable->HasIdentifier(FName(InNewText.ToString())) && !InNewText.EqualTo(FText::FromName(CurrentName)))
+							{
+								OutErrorMessage = LOCTEXT("AttributeNameExists", "Name already exists in the hierarchy.");
+								return false;
+							}
+
+							return true;
+						})
 					.OnTextCommitted(this, &SHierarchyTableRow::OnCommitRename)
 			];
 
diff --git a/Engine/Plugins/Experimental/Animation/HierarchyTable/Source/Editor/Private/SHierarchyTableRow.h b/Engine/Plugins/Experimental/Animation/HierarchyTable/Source/Editor/Private/SHierarchyTableRow.h
--- a/Engine/Plugins/Experimental/Animation/HierarchyTable/Source/Editor/Private/SHierarchyTableRow.h
+++ b/Engine/Plugins/Experimental/Animation/HierarchyTable/Source/Editor/Private/SHierarchyTableRow.h
@@ -11,7 +11,7 @@
 class SHierarchyTableRow : public SMultiColumnTableRow<TSharedPtr<SHierarchyTable::FTreeItem>>
 {
 public:
-	DECLARE_DELEGATE_RetVal_TwoParams(bool /* Success */, FOnRenamed, FName /* Old Name */, FName /* New Name */);
+	DECLARE_DELEGATE_RetVal_OneParam(bool /* Success */, FOnRenamed, FName /* New Name */);
 	DECLARE_DELEGATE_TwoParams(FOnReparented, FName /* Attribute Name */, FName /* New Parent Name */);
 
 	SLATE_BEGIN_ARGS(SHierarchyTableRow) {}
diff --git a/Engine/Plugins/Experimental/Animation/HierarchyTable/Source/Editor/Public/HierarchyTableTypeHandler.h b/Engine/Plugins/Experimental/Animation/HierarchyTable/Source/Editor/Public/HierarchyTableTypeHandler.h
--- a/Engine/Plugins/Experimental/Animation/HierarchyTable/Source/Editor/Public/HierarchyTableTypeHandler.h
+++ b/Engine/Plugins/Experimental/Animation/HierarchyTable/Source/Editor/Public/HierarchyTableTypeHandler.h
@@ -38,6 +38,12 @@ public:
 
 	virtual void ConstructHierarchy() {};
 
+	virtual bool CanRenameEntry(const int32 EntryIndex) const { return false; }
+	virtual bool RenameEntry(const int32 EntryIndex, const FName NewName);
+
+	virtual bool CanRemoveEntry(const int32 EntryIndex) const { return false; }
+	virtual bool RemoveEntry(const int32 EntryIndex);
+
 protected:
 	TObjectPtr<UHierarchyTable> HierarchyTable;
 };
\ No newline at end of file
diff --git a/Engine/Plugins/Experimental/Animation/HierarchyTableAnimation/Source/Editor/Private/SkeletonHierarchyTableTypeHandler.cpp b/Engine/Plugins/Experimental/Animation/HierarchyTableAnimation/Source/Editor/Private/SkeletonHierarchyTableTypeHandler.cpp
--- a/Engine/Plugins/Experimental/Animation/HierarchyTableAnimation/Source/Editor/Private/SkeletonHierarchyTableTypeHandler.cpp
+++ b/Engine/Plugins/Experimental/Animation/HierarchyTableAnimation/Source/Editor/Private/SkeletonHierarchyTableTypeHandler.cpp
@@ -138,34 +138,50 @@ void UHierarchyTable_TableTypeHandler_Skeleton::ExtendContextMenu(FMenuBuilder&
 			})
 	);
 
-	MenuBuilder.AddMenuEntry(
+	MenuBuilder.AddSubMenu(
 		LOCTEXT("AddAttribute_Label", "Add Attribute"),
-		LOCTEXT("AddAttribute_Tooltip", "Add a new attribute entry parented to this entry"),
-		FSlateIcon(),
-		FUIAction(
-			FExecuteAction::CreateLambda([this, SelectedIndex, &HierarchyTableView]()
+		LOCTEXT("AddAttribut_Tooltip", "Add a new attribute entry parented to this entry"),
+		FNewMenuDelegate::CreateLambda([this, SelectedIndex, &HierarchyTableView, &PersonaModule, TableMetadata](FMenuBuilder& InSubMenuBuilder)
+			{
+				InSubMenuBuilder.BeginSection(NAME_None, LOCTEXT("NewAttributeSection", "New Attribute"));
 				{
-					AddAttribute(SelectedIndex, "Test");
-					HierarchyTableView.RegenerateHierarchy();
-				}
-			)
-		));
+					InSubMenuBuilder.AddWidget(
+						SNew(SBox)
+						.Padding(FMargin(8.0f, 0.0f))
+						[
+							SNew(SEditableTextBox)
+								.Text(LOCTEXT("NewAttributeDefault", "NewAttribute"))
+								.OnTextCommitted_Lambda([this, SelectedIndex, &HierarchyTableView](const FText& InText, ETextCommit::Type InTextCommit)
+									{
+										if (InTextCommit == ETextCommit::OnEnter)
+										{
+											const int32 ParentIndex = SelectedIndex == INDEX_NONE ? 0 : SelectedIndex;
+											AddAttribute(ParentIndex, FName(InText.ToString()));
+											HierarchyTableView.RegenerateHierarchy();
+										}
+									})
+								.OnVerifyTextChanged_Lambda([this](const FText& InNewText, FText& OutErrorMessage) -> bool
+									{
+										if (InNewText.IsEmpty())
+										{
+											OutErrorMessage = LOCTEXT("RenameEmpty", "Name can't be empty.");
+											return false;
+										}
 
-	if (EntryMetadata.EntryType == ESkeletonHierarchyTable_TablePayloadEntryType::Curve || EntryMetadata.EntryType == ESkeletonHierarchyTable_TablePayloadEntryType::Attribute)
-	{
-		MenuBuilder.AddMenuEntry(
-			LOCTEXT("RemoveEntry_Label", "Remove Entry"),
-			LOCTEXT("RemoveEntry_Tooltip", "Remove the selected entry"),
-			FSlateIcon(),
-			FUIAction(
-				FExecuteAction::CreateLambda([this, SelectedIndex, &HierarchyTableView]()
-					{
-						RemoveCurve(SelectedIndex);
-						HierarchyTableView.RegenerateHierarchy();
-					}
-				)
-			));
-	}
+										if (HierarchyTable->HasIdentifier(FName(InNewText.ToString())))
+										{
+											OutErrorMessage = LOCTEXT("RenameExists", "Name already exists in the hierarchy.");
+											return false;
+										}
+
+										return true;
+									})
+						],
+						FText::GetEmpty(), true);
+				}
+				InSubMenuBuilder.EndSection();
+			})
+	);
 }
 
 void UHierarchyTable_TableTypeHandler_Skeleton::ConstructHierarchy()
@@ -278,4 +294,30 @@ void UHierarchyTable_TableTypeHandler_Skeleton::AddAttribute(const int32 ParentI
 	HierarchyTable->AddEntry(NewEntry);
 }
 
+bool UHierarchyTable_TableTypeHandler_Skeleton::CanRenameEntry(const int32 EntryIndex) const
+{
+	if (const FHierarchyTableEntryData* const Entry = HierarchyTable->GetTableEntry(EntryIndex))
+	{
+		const FHierarchyTable_TablePayloadType_Skeleton& EntryMetadata = Entry->GetMetadata<FHierarchyTable_TablePayloadType_Skeleton>();
+		const ESkeletonHierarchyTable_TablePayloadEntryType EntryType = EntryMetadata.EntryType;
+
+		return EntryType == ESkeletonHierarchyTable_TablePayloadEntryType::Curve || EntryType == ESkeletonHierarchyTable_TablePayloadEntryType::Attribute;
+	}
+
+	return false;
+}
+
+bool UHierarchyTable_TableTypeHandler_Skeleton::CanRemoveEntry(const int32 EntryIndex) const
+{
+	if (const FHierarchyTableEntryData* const Entry = HierarchyTable->GetTableEntry(EntryIndex))
+	{
+		const FHierarchyTable_TablePayloadType_Skeleton& EntryMetadata = Entry->GetMetadata<FHierarchyTable_TablePayloadType_Skeleton>();
+		const ESkeletonHierarchyTable_TablePayloadEntryType EntryType = EntryMetadata.EntryType;
+
+		return EntryType == ESkeletonHierarchyTable_TablePayloadEntryType::Curve || EntryType == ESkeletonHierarchyTable_TablePayloadEntryType::Attribute;
+	}
+
+	return false;
+}
+
 #undef LOCTEXT_NAMESPACE
\ No newline at end of file
diff --git a/Engine/Plugins/Experimental/Animation/HierarchyTableAnimation/Source/Editor/Private/SkeletonHierarchyTableTypeHandler.h b/Engine/Plugins/Experimental/Animation/HierarchyTableAnimation/Source/Editor/Private/SkeletonHierarchyTableTypeHandler.h
--- a/Engine/Plugins/Experimental/Animation/HierarchyTableAnimation/Source/Editor/Private/SkeletonHierarchyTableTypeHandler.h
+++ b/Engine/Plugins/Experimental/Animation/HierarchyTableAnimation/Source/Editor/Private/SkeletonHierarchyTableTypeHandler.h
@@ -24,6 +24,10 @@ public:
 
 	virtual FSlateColor GetEntryIconColor(const int32 EntryIndex) const override;
 
+	virtual bool CanRenameEntry(const int32 EntryIndex) const override;
+
+	virtual bool CanRemoveEntry(const int32 EntryIndex) const override;
+
 private:
 	void AddCurve(const int32 ParentIndex, const FName Identifier) const;
 	void RemoveCurve(const int32 Index) const;
