Fix invalid memory access on editor shutdown in HierarchyTable plugin

diff --git a/Engine/Plugins/Experimental/Animation/HierarchyTable/Source/Editor/Private/HierachyTableEditorModule.cpp b/Engine/Plugins/Experimental/Animation/HierarchyTable/Source/Editor/Private/HierachyTableEditorModule.cpp
--- a/Engine/Plugins/Experimental/Animation/HierarchyTable/Source/Editor/Private/HierachyTableEditorModule.cpp
+++ b/Engine/Plugins/Experimental/Animation/HierarchyTable/Source/Editor/Private/HierachyTableEditorModule.cpp
@@ -7,12 +7,19 @@
 
 void FHierarchyTableEditorModule::StartupModule()
 {
+	BuiltinTableTypes.Add(FHierarchyTable_TableType_Default::StaticStruct());
 	RegisterTableType(FHierarchyTable_TableType_Default::StaticStruct(), UHierarchyTable_TableTypeHandler_Default::StaticClass());
 }
 
 void FHierarchyTableEditorModule::ShutdownModule()
 {
-	UnregisterTableType(FHierarchyTable_TableType_Default::StaticStruct());
+	for (TWeakObjectPtr<UScriptStruct> WeakPtr : BuiltinTableTypes)
+	{
+		if (const UScriptStruct* StructPtr = WeakPtr.Get())
+		{
+			UnregisterTableType(StructPtr);
+		}
+	}
 }
 
 void FHierarchyTableEditorModule::RegisterTableType(const UScriptStruct* TableType, const UClass* Handler)
diff --git a/Engine/Plugins/Experimental/Animation/HierarchyTable/Source/Editor/Public/HierarchyTableEditorModule.h b/Engine/Plugins/Experimental/Animation/HierarchyTable/Source/Editor/Public/HierarchyTableEditorModule.h
--- a/Engine/Plugins/Experimental/Animation/HierarchyTable/Source/Editor/Public/HierarchyTableEditorModule.h
+++ b/Engine/Plugins/Experimental/Animation/HierarchyTable/Source/Editor/Public/HierarchyTableEditorModule.h
@@ -6,6 +6,7 @@
 #include "HierarchyTableType.h"
 #include "Modules/ModuleInterface.h"
 #include "HierarchyTable.h"
+#include "UObject/ObjectKey.h"
 
 struct IHierarchyTableColumn;
 
@@ -31,7 +32,9 @@ public:
 	TArray<TSharedPtr<IHierarchyTableColumn>> GetElementTypeEditorColumns(const TObjectPtr<const UHierarchyTable> HierarchyTable);
 
 private:
-	TMap<const UScriptStruct*, const UClass*> TableHandlers;
+	TArray<TWeakObjectPtr<UScriptStruct>> BuiltinTableTypes;
 
-	TMap<TObjectPtr<const UScriptStruct>, TArray<TSharedPtr<IHierarchyTableColumn>>> EditorColumns;
+	TMap<TObjectKey<UScriptStruct>, const UClass*> TableHandlers;
+
+	TMap<TObjectKey<UScriptStruct>, TArray<TSharedPtr<IHierarchyTableColumn>>> EditorColumns;
 };
\ No newline at end of file
diff --git a/Engine/Plugins/Experimental/Animation/HierarchyTableAnimation/Source/Editor/Private/HierachyTableAnimationEditorModule.cpp b/Engine/Plugins/Experimental/Animation/HierarchyTableAnimation/Source/Editor/Private/HierachyTableAnimationEditorModule.cpp
--- a/Engine/Plugins/Experimental/Animation/HierarchyTableAnimation/Source/Editor/Private/HierachyTableAnimationEditorModule.cpp
+++ b/Engine/Plugins/Experimental/Animation/HierarchyTableAnimation/Source/Editor/Private/HierachyTableAnimationEditorModule.cpp
@@ -14,13 +14,16 @@ void FHierarchyTableAnimationEditorModule::StartupModule()
 {
 	FHierarchyTableEditorModule& HierarchyTableModule = FModuleManager::LoadModuleChecked<FHierarchyTableEditorModule>("HierarchyTableEditor");
 	
+	BuiltinTableTypes.Add(FHierarchyTable_TableType_Skeleton::StaticStruct());
 	HierarchyTableModule.RegisterTableType(FHierarchyTable_TableType_Skeleton::StaticStruct(), UHierarchyTable_TableTypeHandler_Skeleton::StaticClass());
 
+	BuiltinElementTypes.Add(FHierarchyTable_ElementType_Mask::StaticStruct());
 	HierarchyTableModule.RegisterElementTypeEditorColumns(FHierarchyTable_ElementType_Mask::StaticStruct(),
 		{
 			MakeShared<FHierarchyTableColumn_Mask>(),
 		});
 
+	BuiltinElementTypes.Add(FHierarchyTable_ElementType_Time::StaticStruct());
 	HierarchyTableModule.RegisterElementTypeEditorColumns(FHierarchyTable_ElementType_Time::StaticStruct(),
 		{
 			MakeShared<FHierarchyTableColumn_TimeProfile_StartTime>(),
@@ -35,9 +38,22 @@ void FHierarchyTableAnimationEditorModule::ShutdownModule()
 	if (FModuleManager::Get().IsModuleLoaded("HierarchyTableEditor"))
 	{
 		FHierarchyTableEditorModule& HierarchyTableModule = FModuleManager::GetModuleChecked<FHierarchyTableEditorModule>("HierarchyTableEditor");
-		HierarchyTableModule.UnregisterTableType(FHierarchyTable_TableType_Skeleton::StaticStruct());
-		HierarchyTableModule.UnregisterElementTypeEditorColumns(FHierarchyTable_ElementType_Mask::StaticStruct());
-		HierarchyTableModule.UnregisterElementTypeEditorColumns(FHierarchyTable_ElementType_Time::StaticStruct());
+
+		for (TWeakObjectPtr<UScriptStruct> WeakPtr : BuiltinTableTypes)
+		{
+			if (const UScriptStruct* StructPtr = WeakPtr.Get())
+			{
+				HierarchyTableModule.UnregisterTableType(StructPtr);
+			}
+		}
+
+		for (TWeakObjectPtr<UScriptStruct> WeakPtr : BuiltinElementTypes)
+		{
+			if (const UScriptStruct* StructPtr = WeakPtr.Get())
+			{
+				HierarchyTableModule.UnregisterElementTypeEditorColumns(StructPtr);
+			}
+		}
 	}
 }
 
diff --git a/Engine/Plugins/Experimental/Animation/HierarchyTableAnimation/Source/Editor/Public/HierarchyTableAnimationEditorModule.h b/Engine/Plugins/Experimental/Animation/HierarchyTableAnimation/Source/Editor/Public/HierarchyTableAnimationEditorModule.h
--- a/Engine/Plugins/Experimental/Animation/HierarchyTableAnimation/Source/Editor/Public/HierarchyTableAnimationEditorModule.h
+++ b/Engine/Plugins/Experimental/Animation/HierarchyTableAnimation/Source/Editor/Public/HierarchyTableAnimationEditorModule.h
@@ -2,6 +2,7 @@
 
 #pragma once
 
+#include "CoreMinimal.h"
 #include "Modules/ModuleInterface.h"
 
 class FHierarchyTableAnimationEditorModule : public IModuleInterface
@@ -10,4 +11,8 @@ public:
 	virtual void StartupModule() override;
 	
 	virtual void ShutdownModule() override;
+
+private:
+	TArray<TWeakObjectPtr<UScriptStruct>> BuiltinTableTypes;
+	TArray<TWeakObjectPtr<UScriptStruct>> BuiltinElementTypes;
 };
\ No newline at end of file
