Hierarchy Tables - Allow switching the bound skeleton

diff --git a/Engine/Plugins/Experimental/Animation/HierarchyTable/Source/Editor/Private/SHierarchyTable.cpp b/Engine/Plugins/Experimental/Animation/HierarchyTable/Source/Editor/Private/SHierarchyTable.cpp
--- a/Engine/Plugins/Experimental/Animation/HierarchyTable/Source/Editor/Private/SHierarchyTable.cpp
+++ b/Engine/Plugins/Experimental/Animation/HierarchyTable/Source/Editor/Private/SHierarchyTable.cpp
@@ -24,6 +24,11 @@ void SHierarchyTable::Construct(const FArguments& InArgs, TObjectPtr<UHierarchyT
 
 	HierarchyTable = InHierarchyTable;
 
+	OnHierarchyChangedHandle = HierarchyTable->RegisterOnHierarchyChanged(UHierarchyTable::FOnHierarchyChanged::FDelegate::CreateLambda([this]()
+	{
+		RegenerateTreeViewItems();
+	}));
+
 	TSharedPtr<SHeaderRow> HeaderRow;
 
 	ChildSlot
@@ -78,6 +83,14 @@ void SHierarchyTable::Construct(const FArguments& InArgs, TObjectPtr<UHierarchyT
 	}
 }
 
+SHierarchyTable::~SHierarchyTable()
+{
+	if (HierarchyTable)
+	{
+		HierarchyTable->UnregisterOnHierarchyChanged(OnHierarchyChangedHandle);
+	}
+}
+
 int32 SHierarchyTable::GetSelectedEntryIndex() const
 {
 	TArray<TSharedPtr<FTreeItem>> SelectedItems;
@@ -91,11 +104,6 @@ int32 SHierarchyTable::GetSelectedEntryIndex() const
 	return INDEX_NONE;
 }
 
-void SHierarchyTable::RegenerateHierarchy()
-{
-	RegenerateTreeViewItems();
-}
-
 void SHierarchyTable::TreeView_HandleGetChildren(TSharedPtr<FTreeItem> InItem, TArray<TSharedPtr<FTreeItem>>& OutChildren) const
 {
 	OutChildren.Append(InItem->Children);
diff --git a/Engine/Plugins/Experimental/Animation/HierarchyTable/Source/Editor/Private/SHierarchyTable.h b/Engine/Plugins/Experimental/Animation/HierarchyTable/Source/Editor/Private/SHierarchyTable.h
--- a/Engine/Plugins/Experimental/Animation/HierarchyTable/Source/Editor/Private/SHierarchyTable.h
+++ b/Engine/Plugins/Experimental/Animation/HierarchyTable/Source/Editor/Private/SHierarchyTable.h
@@ -39,10 +39,10 @@ public:
 	};
 
 	void Construct(const FArguments& InArgs, TObjectPtr<UHierarchyTable> InHierarchyTable);
+	~SHierarchyTable();
 
 	// IHierarchyTable
 	int32 GetSelectedEntryIndex() const override;
-	void RegenerateHierarchy() override;
 
 	// FSelfRegisteringEditorUndoClient
 	void PostUndo(bool bSuccess) override;
@@ -72,4 +72,6 @@ private:
 	TObjectPtr<UHierarchyTable> HierarchyTable;
 
 	TObjectPtr<UHierarchyTable_TableTypeHandler> TableHandler;
+
+	FDelegateHandle OnHierarchyChangedHandle;
 };
diff --git a/Engine/Plugins/Experimental/Animation/HierarchyTable/Source/Editor/Public/HierarchyTableTypeHandler.h b/Engine/Plugins/Experimental/Animation/HierarchyTable/Source/Editor/Public/HierarchyTableTypeHandler.h
--- a/Engine/Plugins/Experimental/Animation/HierarchyTable/Source/Editor/Public/HierarchyTableTypeHandler.h
+++ b/Engine/Plugins/Experimental/Animation/HierarchyTable/Source/Editor/Public/HierarchyTableTypeHandler.h
@@ -24,18 +24,22 @@ class HIERARCHYTABLEEDITOR_API UHierarchyTable_TableTypeHandler : public UObject
 public:
 	void SetHierarchyTable(TObjectPtr<UHierarchyTable> InHierarchyTable) { HierarchyTable = InHierarchyTable; }
 
-	virtual void InitialiseTable(const FInstancedStruct& TableType) const {}
-
+	/** Override to provide additional configurable properties when creating a new hierarchy table of this table type */
 	virtual bool FactoryConfigureProperties(FInstancedStruct& TableType) const { return false; }
 
+	/** Override to add additional buttons to the editor toolbar */
 	virtual void ExtendToolbar(UToolMenu* ToolMenu, IHierarchyTable& HierarchyTableView) const {}
 
+	/** Override to customize the context menu constructed for a particular entry in the table */
 	virtual void ExtendContextMenu(FMenuBuilder& MenuBuilder, IHierarchyTable& HierarchyTableView) const {}
 
+	/** Override to provide an icon for a particular entry in the table */
 	virtual FSlateIcon GetEntryIcon(const int32 EntryIndex) const;
-
+	
+	/** Override to provide an icon color for a particular entry in the table */
 	virtual FSlateColor GetEntryIconColor(const int32 EntryIndex) const;
 
+	/** Resets and builds the hierarchy table data according to the current table metadata */
 	virtual void ConstructHierarchy() {};
 
 	virtual bool CanRenameEntry(const int32 EntryIndex) const { return false; }
diff --git a/Engine/Plugins/Experimental/Animation/HierarchyTable/Source/Editor/Public/IHierarchyTable.h b/Engine/Plugins/Experimental/Animation/HierarchyTable/Source/Editor/Public/IHierarchyTable.h
--- a/Engine/Plugins/Experimental/Animation/HierarchyTable/Source/Editor/Public/IHierarchyTable.h
+++ b/Engine/Plugins/Experimental/Animation/HierarchyTable/Source/Editor/Public/IHierarchyTable.h
@@ -8,6 +8,4 @@ class IHierarchyTable
 {
 public:
 	virtual int32 GetSelectedEntryIndex() const = 0;
-
-	virtual void RegenerateHierarchy() = 0;
 };
\ No newline at end of file
diff --git a/Engine/Plugins/Experimental/Animation/HierarchyTable/Source/Runtime/Private/HierarchyTable.cpp b/Engine/Plugins/Experimental/Animation/HierarchyTable/Source/Runtime/Private/HierarchyTable.cpp
--- a/Engine/Plugins/Experimental/Animation/HierarchyTable/Source/Runtime/Private/HierarchyTable.cpp
+++ b/Engine/Plugins/Experimental/Animation/HierarchyTable/Source/Runtime/Private/HierarchyTable.cpp
@@ -105,11 +105,26 @@ int32 UHierarchyTable::AddEntry(const FHierarchyTableEntryData& Entry)
 
 	const int32 EntryIndex = TableData.Add(Entry);
 #if WITH_EDITOR
-	OnEntryAdded.Broadcast(EntryIndex);
+	OnHierarchyChanged.Broadcast();
 #endif //WITH_EDITOR
 	return EntryIndex;
 }
 
+void UHierarchyTable::AddBulkEntries(const TConstArrayView<FHierarchyTableEntryData> Entries)
+{
+	for (const FHierarchyTableEntryData& Entry : Entries)
+	{
+		if (!HasIdentifier(Entry.Identifier))
+		{
+			TableData.Add(Entry);
+		}
+	}
+
+#if WITH_EDITOR
+	OnHierarchyChanged.Broadcast();
+#endif //WITH_EDITOR
+}
+
 void UHierarchyTable::RemoveEntry(const int32 IndexToRemove)
 {
 	if (!TableData.IsValidIndex(IndexToRemove))
@@ -133,17 +148,21 @@ void UHierarchyTable::RemoveEntry(const int32 IndexToRemove)
 	}
 
 	TableData.RemoveAt(IndexToRemove);
+
+#if WITH_EDITOR
+	OnHierarchyChanged.Broadcast();
+#endif //WITH_EDITOR
 }
 
 #if WITH_EDITOR
-FDelegateHandle UHierarchyTable::RegisterOnEntryAdded(const FOnEntryAdded::FDelegate& Delegate)
+FDelegateHandle UHierarchyTable::RegisterOnHierarchyChanged(const FOnHierarchyChanged::FDelegate& Delegate)
 {
-	return OnEntryAdded.Add(Delegate);
+	return OnHierarchyChanged.Add(Delegate);
 }
 
-bool UHierarchyTable::UnregisterOnEntryAdded(const FDelegateHandle& Handle)
+bool UHierarchyTable::UnregisterOnHierarchyChanged(const FDelegateHandle& Handle)
 {
-	return OnEntryAdded.Remove(Handle);
+	return OnHierarchyChanged.Remove(Handle);
 }
 #endif //WITH_EDITOR
 
@@ -171,4 +190,13 @@ FHierarchyTableEntryData* const UHierarchyTable::GetMutableTableEntry(const int3
 		return &TableData[EntryIndex];
 	}
 	return nullptr;
+}
+
+void UHierarchyTable::EmptyTable()
+{
+	TableData.Empty();
+
+#if WITH_EDITOR
+	OnHierarchyChanged.Broadcast();
+#endif //WITH_EDITOR
 }
\ No newline at end of file
diff --git a/Engine/Plugins/Experimental/Animation/HierarchyTable/Source/Runtime/Public/HierarchyTable.h b/Engine/Plugins/Experimental/Animation/HierarchyTable/Source/Runtime/Public/HierarchyTable.h
--- a/Engine/Plugins/Experimental/Animation/HierarchyTable/Source/Runtime/Public/HierarchyTable.h
+++ b/Engine/Plugins/Experimental/Animation/HierarchyTable/Source/Runtime/Public/HierarchyTable.h
@@ -105,7 +105,7 @@ public:
 
 private:
 	// The table metadata stores any data dependent on the table type needed to create and maintain the hierarchy.
-	UPROPERTY(VisibleAnywhere, Category = Default)
+	UPROPERTY(EditAnywhere, Category = Default)
 	FInstancedStruct TableMetadata;
 
 	// The element type is the type that each element is mapped into.
@@ -156,12 +156,16 @@ public:
 
 	const TArray<FHierarchyTableEntryData>& GetTableData() const { return TableData; }
 
+	void EmptyTable();
+
 	const FHierarchyTableEntryData* const GetTableEntry(const int32 EntryIndex) const;
 
 	FHierarchyTableEntryData* const GetMutableTableEntry(const int32 EntryIndex);
 
 	int32 AddEntry(const FHierarchyTableEntryData& Entry);
 
+	void AddBulkEntries(const TConstArrayView<FHierarchyTableEntryData> Entries);
+
 	void RemoveEntry(const int32 Index);
 
 	TArray<const FHierarchyTableEntryData*> GetChildren(const FHierarchyTableEntryData& Parent) const;
@@ -178,12 +182,12 @@ public:
 	}
 
 #if WITH_EDITOR
-	DECLARE_MULTICAST_DELEGATE_OneParam(FOnEntryAdded, int32);
+	DECLARE_MULTICAST_DELEGATE(FOnHierarchyChanged);
 private:
-	FOnEntryAdded OnEntryAdded;
+	FOnHierarchyChanged OnHierarchyChanged;
 
 public:
-	FDelegateHandle RegisterOnEntryAdded(const FOnEntryAdded::FDelegate& Delegate);
-	bool UnregisterOnEntryAdded(const FDelegateHandle& Handle);
+	FDelegateHandle RegisterOnHierarchyChanged(const FOnHierarchyChanged::FDelegate& Delegate);
+	bool UnregisterOnHierarchyChanged(const FDelegateHandle& Handle);
 #endif
 };
diff --git a/Engine/Plugins/Experimental/Animation/HierarchyTableAnimation/Source/Editor/Private/HierachyTableAnimationEditorModule.cpp b/Engine/Plugins/Experimental/Animation/HierarchyTableAnimation/Source/Editor/Private/HierachyTableAnimationEditorModule.cpp
--- a/Engine/Plugins/Experimental/Animation/HierarchyTableAnimation/Source/Editor/Private/HierachyTableAnimationEditorModule.cpp
+++ b/Engine/Plugins/Experimental/Animation/HierarchyTableAnimation/Source/Editor/Private/HierachyTableAnimationEditorModule.cpp
@@ -9,6 +9,8 @@
 #include "SkeletonHierarchyTableTypeHandler.h"
 #include "MaskProfile/MaskProfileColumn.h"
 #include "TimeProfile/TimeProfileColumn.h"
+#include "SkeletonHierarchyTableTypeDetailsCustomization.h"
+#include "PropertyEditorModule.h"
 
 void FHierarchyTableAnimationEditorModule::StartupModule()
 {
@@ -31,6 +33,9 @@ void FHierarchyTableAnimationEditorModule::StartupModule()
 			MakeShared<FHierarchyTableColumn_TimeProfile_TimeFactor>(),
 			MakeShared<FHierarchyTableColumn_TimeProfile_Preview>(),
 		});
+
+	FPropertyEditorModule& PropertyModule = FModuleManager::LoadModuleChecked<FPropertyEditorModule>("PropertyEditor");
+	PropertyModule.RegisterCustomPropertyTypeLayout(FHierarchyTable_TableType_Skeleton::StaticStruct()->GetFName(), FOnGetPropertyTypeCustomizationInstance::CreateStatic(&FHierarchyTableSkeletonTableTypeDetailsCustomization::MakeInstance));
 }
 
 void FHierarchyTableAnimationEditorModule::ShutdownModule()
diff --git a/Engine/Plugins/Experimental/Animation/HierarchyTableAnimation/Source/Editor/Private/SkeletonHierarchyTableTypeDetailsCustomization.cpp b/Engine/Plugins/Experimental/Animation/HierarchyTableAnimation/Source/Editor/Private/SkeletonHierarchyTableTypeDetailsCustomization.cpp
new file mode 100644
--- /dev/null
+++ b/Engine/Plugins/Experimental/Animation/HierarchyTableAnimation/Source/Editor/Private/SkeletonHierarchyTableTypeDetailsCustomization.cpp
@@ -0,0 +1,60 @@
+// Copyright Epic Games, Inc. All Rights Reserved.
+
+#include "SkeletonHierarchyTableTypeDetailsCustomization.h"
+#include "DetailLayoutBuilder.h"
+#include "HierarchyTable.h"
+#include "HierarchyTableEditorModule.h"
+#include "SkeletonHierarchyTableType.h"
+#include "IDetailChildrenBuilder.h"
+#include "Modules/ModuleManager.h"
+
+TSharedRef<IPropertyTypeCustomization> FHierarchyTableSkeletonTableTypeDetailsCustomization::MakeInstance()
+{
+	return MakeShareable(new FHierarchyTableSkeletonTableTypeDetailsCustomization());
+}
+
+void FHierarchyTableSkeletonTableTypeDetailsCustomization::CustomizeHeader(TSharedRef<IPropertyHandle> StructPropertyHandle, class FDetailWidgetRow& HeaderRow, IPropertyTypeCustomizationUtils& StructCustomizationUtils)
+{
+}
+
+void FHierarchyTableSkeletonTableTypeDetailsCustomization::CustomizeChildren(TSharedRef<class IPropertyHandle> StructPropertyHandle, class IDetailChildrenBuilder& StructBuilder, IPropertyTypeCustomizationUtils& StructCustomizationUtils)
+{
+	TArray<UObject*> EditingObjects;
+	StructPropertyHandle->GetOuterObjects(EditingObjects);
+
+	if (EditingObjects.Num() != 1)
+	{
+		return;
+	}
+
+	TObjectPtr<UHierarchyTable> HierarchyTable = CastChecked<UHierarchyTable>(EditingObjects[0]);
+
+	if (StructPropertyHandle->IsValidHandle())
+	{
+		uint32 NumChildren = 0;
+
+		StructPropertyHandle->GetNumChildren(NumChildren);
+
+		for (uint32 ChildIndex = 0; ChildIndex < NumChildren; ChildIndex++)
+		{
+			TSharedRef<IPropertyHandle> ChildHandle = StructPropertyHandle->GetChildHandle(ChildIndex).ToSharedRef();
+
+			if (ChildHandle->IsValidHandle())
+			{
+				if (ChildHandle->GetProperty()->GetFName() == GET_MEMBER_NAME_CHECKED(FHierarchyTable_TableType_Skeleton, Skeleton))
+				{
+					ChildHandle->SetOnPropertyValueChanged(FSimpleDelegate::CreateLambda([HierarchyTable]()
+						{
+							FHierarchyTableEditorModule& HierarchyTableModule = FModuleManager::LoadModuleChecked<FHierarchyTableEditorModule>("HierarchyTableEditor");
+							const TObjectPtr<UHierarchyTable_TableTypeHandler> HierarchyTableHandler = HierarchyTableModule.CreateTableHandler(HierarchyTable);
+							check(HierarchyTableHandler);
+
+							HierarchyTableHandler->ConstructHierarchy();
+						}));
+				}
+
+				StructBuilder.AddProperty(ChildHandle);
+			}
+		}
+	}
+}
\ No newline at end of file
diff --git a/Engine/Plugins/Experimental/Animation/HierarchyTableAnimation/Source/Editor/Private/SkeletonHierarchyTableTypeDetailsCustomization.h b/Engine/Plugins/Experimental/Animation/HierarchyTableAnimation/Source/Editor/Private/SkeletonHierarchyTableTypeDetailsCustomization.h
new file mode 100644
--- /dev/null
+++ b/Engine/Plugins/Experimental/Animation/HierarchyTableAnimation/Source/Editor/Private/SkeletonHierarchyTableTypeDetailsCustomization.h
@@ -0,0 +1,17 @@
+// Copyright Epic Games, Inc. All Rights Reserved.
+
+#pragma once
+
+#include "IPropertyTypeCustomization.h"
+
+class IDetailLayoutBuilder;
+
+class FHierarchyTableSkeletonTableTypeDetailsCustomization : public IPropertyTypeCustomization
+{
+public:
+	static TSharedRef<IPropertyTypeCustomization> MakeInstance();
+
+	// IPropertyTypeCustomization
+	virtual void CustomizeHeader(TSharedRef<IPropertyHandle> StructPropertyHandle, class FDetailWidgetRow& HeaderRow, IPropertyTypeCustomizationUtils& StructCustomizationUtils) override;
+	virtual void CustomizeChildren(TSharedRef<class IPropertyHandle> StructPropertyHandle, class IDetailChildrenBuilder& StructBuilder, IPropertyTypeCustomizationUtils& StructCustomizationUtils) override;
+};
diff --git a/Engine/Plugins/Experimental/Animation/HierarchyTableAnimation/Source/Editor/Private/SkeletonHierarchyTableTypeHandler.cpp b/Engine/Plugins/Experimental/Animation/HierarchyTableAnimation/Source/Editor/Private/SkeletonHierarchyTableTypeHandler.cpp
--- a/Engine/Plugins/Experimental/Animation/HierarchyTableAnimation/Source/Editor/Private/SkeletonHierarchyTableTypeHandler.cpp
+++ b/Engine/Plugins/Experimental/Animation/HierarchyTableAnimation/Source/Editor/Private/SkeletonHierarchyTableTypeHandler.cpp
@@ -92,7 +92,6 @@ void UHierarchyTable_TableTypeHandler_Skeleton::ExtendContextMenu(FMenuBuilder&
 										{
 											const int32 ParentIndex = SelectedIndex == INDEX_NONE ? 0 : SelectedIndex;
 											AddCurve(ParentIndex, FName(InText.ToString()));
-											HierarchyTableView.RegenerateHierarchy();
 										}
 									})
 								.OnVerifyTextChanged_Lambda([this](const FText& InNewText, FText& OutErrorMessage) -> bool
@@ -129,7 +128,6 @@ void UHierarchyTable_TableTypeHandler_Skeleton::ExtendContextMenu(FMenuBuilder&
 									{
 										AddCurve(ParentIndex, Curve);
 									}
-									HierarchyTableView.RegenerateHierarchy();
 								}),
 							FIsCurveNameMarkedForExclusion::CreateLambda([this](const FName& InName)
 								{
@@ -160,7 +158,6 @@ void UHierarchyTable_TableTypeHandler_Skeleton::ExtendContextMenu(FMenuBuilder&
 										{
 											const int32 ParentIndex = SelectedIndex == INDEX_NONE ? 0 : SelectedIndex;
 											AddAttribute(ParentIndex, FName(InText.ToString()));
-											HierarchyTableView.RegenerateHierarchy();
 										}
 									})
 								.OnVerifyTextChanged_Lambda([this](const FText& InNewText, FText& OutErrorMessage) -> bool
@@ -189,10 +186,10 @@ void UHierarchyTable_TableTypeHandler_Skeleton::ExtendContextMenu(FMenuBuilder&
 
 void UHierarchyTable_TableTypeHandler_Skeleton::ConstructHierarchy()
 {
-	const FHierarchyTable_TableType_Skeleton& SkeletonTableType = HierarchyTable->GetTableMetadata<FHierarchyTable_TableType_Skeleton>();
+	HierarchyTable->EmptyTable();
 
+	const FHierarchyTable_TableType_Skeleton& SkeletonTableType = HierarchyTable->GetTableMetadata<FHierarchyTable_TableType_Skeleton>();
 	check(SkeletonTableType.Skeleton);
-	check(HierarchyTable->GetTableData().Num() == 0);
 
 	FInstancedStruct DefaultEntry = HierarchyTable->CreateDefaultValue();
 
@@ -202,6 +199,8 @@ void UHierarchyTable_TableTypeHandler_Skeleton::ConstructHierarchy()
 
 	FReferenceSkeleton RefSkeleton = SkeletonTableType.Skeleton->GetReferenceSkeleton();
 
+	TArray<FHierarchyTableEntryData> EntriesToAdd;
+
 	for (int32 BoneIndex = 0; BoneIndex < RefSkeleton.GetNum(); ++BoneIndex)
 	{
 		FHierarchyTableEntryData EntryData;
@@ -210,8 +209,11 @@ void UHierarchyTable_TableTypeHandler_Skeleton::ConstructHierarchy()
 		EntryData.TablePayload = DefaultTablePayload;
 		EntryData.Payload = BoneIndex == 0 ? DefaultEntry : TOptional<FInstancedStruct>();
 		EntryData.OwnerTable = HierarchyTable;
-		HierarchyTable->AddEntry(EntryData);
+
+		EntriesToAdd.Add(EntryData);
 	}
+
+	HierarchyTable->AddBulkEntries(EntriesToAdd);
 }
 
 bool UHierarchyTable_TableTypeHandler_Skeleton::FactoryConfigureProperties(FInstancedStruct& TableType) const
diff --git a/Engine/Plugins/Experimental/Animation/HierarchyTableAnimation/Source/Runtime/Public/SkeletonHierarchyTableType.h b/Engine/Plugins/Experimental/Animation/HierarchyTableAnimation/Source/Runtime/Public/SkeletonHierarchyTableType.h
--- a/Engine/Plugins/Experimental/Animation/HierarchyTableAnimation/Source/Runtime/Public/SkeletonHierarchyTableType.h
+++ b/Engine/Plugins/Experimental/Animation/HierarchyTableAnimation/Source/Runtime/Public/SkeletonHierarchyTableType.h
@@ -8,12 +8,12 @@
 
 class USkeleton;
 
-USTRUCT()
+USTRUCT(DisplayName = "Skeleton")
 struct FHierarchyTable_TableType_Skeleton final : public FHierarchyTable_TableType
 {
 	GENERATED_BODY()
 
-	UPROPERTY()
+	UPROPERTY(EditAnywhere, Category = Default)
 	TObjectPtr<USkeleton> Skeleton;
 };
 
