Hierarchy Tables - Add missing undo transactions for adding curves and attributes

diff --git a/Engine/Plugins/Experimental/Animation/HierarchyTableAnimation/Source/Editor/Private/SkeletonHierarchyTableTypeHandler.cpp b/Engine/Plugins/Experimental/Animation/HierarchyTableAnimation/Source/Editor/Private/SkeletonHierarchyTableTypeHandler.cpp
--- a/Engine/Plugins/Experimental/Animation/HierarchyTableAnimation/Source/Editor/Private/SkeletonHierarchyTableTypeHandler.cpp
+++ b/Engine/Plugins/Experimental/Animation/HierarchyTableAnimation/Source/Editor/Private/SkeletonHierarchyTableTypeHandler.cpp
@@ -21,6 +21,7 @@
 #include "Widgets/Text/STextBlock.h"
 #include "Widgets/SBoxPanel.h"
 #include "Widgets/Input/SEditableTextBox.h"
+#include "ScopedTransaction.h"
 
 #define LOCTEXT_NAMESPACE "UHierarchyTable_TableTypeHandler_Skeleton"
 
@@ -258,6 +259,9 @@ void UHierarchyTable_TableTypeHandler_Skeleton::ExtendToolbar(UToolMenu* ToolMen
 
 void UHierarchyTable_TableTypeHandler_Skeleton::AddCurve(const int32 ParentIndex, const FName Identifier) const
 {
+	FScopedTransaction AddCurveTransaction(LOCTEXT("AddCurveTransaction", "Add Curve Entry"));
+	HierarchyTable->Modify();
+
 	FHierarchyTableEntryData NewEntry;
 	NewEntry.Parent = ParentIndex;
 	NewEntry.Identifier = Identifier;
@@ -274,14 +278,12 @@ void UHierarchyTable_TableTypeHandler_Skeleton::AddCurve(const int32 ParentIndex
 
 	HierarchyTable->AddEntry(NewEntry);
 }
-
-void UHierarchyTable_TableTypeHandler_Skeleton::RemoveCurve(const int32 Index) const
-{
-	HierarchyTable->RemoveEntry(Index);
-}
 	
 void UHierarchyTable_TableTypeHandler_Skeleton::AddAttribute(const int32 ParentIndex, const FName Identifier) const
 {
+	FScopedTransaction AddAttributeTransaction(LOCTEXT("AddAttributeTransaction", "Add Attribute Entry"));
+	HierarchyTable->Modify();
+
 	FHierarchyTableEntryData NewEntry;
 	NewEntry.Parent = ParentIndex;
 	NewEntry.Identifier = Identifier;
diff --git a/Engine/Plugins/Experimental/Animation/HierarchyTableAnimation/Source/Editor/Private/SkeletonHierarchyTableTypeHandler.h b/Engine/Plugins/Experimental/Animation/HierarchyTableAnimation/Source/Editor/Private/SkeletonHierarchyTableTypeHandler.h
--- a/Engine/Plugins/Experimental/Animation/HierarchyTableAnimation/Source/Editor/Private/SkeletonHierarchyTableTypeHandler.h
+++ b/Engine/Plugins/Experimental/Animation/HierarchyTableAnimation/Source/Editor/Private/SkeletonHierarchyTableTypeHandler.h
@@ -30,7 +30,6 @@ public:
 
 private:
 	void AddCurve(const int32 ParentIndex, const FName Identifier) const;
-	void RemoveCurve(const int32 Index) const;
-	
+
 	void AddAttribute(const int32 ParentIndex, const FName Identifier) const;
 };
\ No newline at end of file
