Add display names for some Hierarchy Table classes

diff --git a/Engine/Plugins/Experimental/Animation/HierarchyTable/Source/Runtime/Public/DefaultHierarchyTableType.h b/Engine/Plugins/Experimental/Animation/HierarchyTable/Source/Runtime/Public/DefaultHierarchyTableType.h
--- a/Engine/Plugins/Experimental/Animation/HierarchyTable/Source/Runtime/Public/DefaultHierarchyTableType.h
+++ b/Engine/Plugins/Experimental/Animation/HierarchyTable/Source/Runtime/Public/DefaultHierarchyTableType.h
@@ -6,7 +6,7 @@
 
 #include "DefaultHierarchyTableType.generated.h"
 
-USTRUCT()
+USTRUCT(DisplayName = "Default")
 struct FHierarchyTable_TableType_Default final : public FHierarchyTable_TableType
 {
 	GENERATED_BODY()
diff --git a/Engine/Plugins/Experimental/Animation/HierarchyTableAnimation/Source/Runtime/Public/MaskProfile/HierarchyTableTypeMask.h b/Engine/Plugins/Experimental/Animation/HierarchyTableAnimation/Source/Runtime/Public/MaskProfile/HierarchyTableTypeMask.h
--- a/Engine/Plugins/Experimental/Animation/HierarchyTableAnimation/Source/Runtime/Public/MaskProfile/HierarchyTableTypeMask.h
+++ b/Engine/Plugins/Experimental/Animation/HierarchyTableAnimation/Source/Runtime/Public/MaskProfile/HierarchyTableTypeMask.h
@@ -6,7 +6,7 @@
 
 #include "HierarchyTableTypeMask.generated.h"
 
-USTRUCT()
+USTRUCT(DisplayName = "Mask")
 struct FHierarchyTable_ElementType_Mask final : public FHierarchyTable_ElementType
 {
 	GENERATED_BODY()
