Fix uninitialized property in FHierarchyTable_TablePayloadType_Skeleton

diff --git a/Engine/Plugins/Experimental/Animation/HierarchyTableAnimation/Source/Runtime/Public/SkeletonHierarchyTableType.h b/Engine/Plugins/Experimental/Animation/HierarchyTableAnimation/Source/Runtime/Public/SkeletonHierarchyTableType.h
--- a/Engine/Plugins/Experimental/Animation/HierarchyTableAnimation/Source/Runtime/Public/SkeletonHierarchyTableType.h
+++ b/Engine/Plugins/Experimental/Animation/HierarchyTableAnimation/Source/Runtime/Public/SkeletonHierarchyTableType.h
@@ -31,5 +31,5 @@ struct FHierarchyTable_TablePayloadType_Skeleton final : public FHierarchyTable_
 	GENERATED_BODY()
 
 	UPROPERTY()
-	ESkeletonHierarchyTable_TablePayloadEntryType EntryType;
+	ESkeletonHierarchyTable_TablePayloadEntryType EntryType = ESkeletonHierarchyTable_TablePayloadEntryType::Bone;
 };
\ No newline at end of file
