Fix profile blend validation throwing errors for valid setups

diff --git a/Engine/Plugins/Experimental/Animation/HierarchyTableAnimation/Source/UncookedOnly/Private/AnimGraphNode_BlendProfileLayeredBlend.cpp b/Engine/Plugins/Experimental/Animation/HierarchyTableAnimation/Source/UncookedOnly/Private/AnimGraphNode_BlendProfileLayeredBlend.cpp
--- a/Engine/Plugins/Experimental/Animation/HierarchyTableAnimation/Source/UncookedOnly/Private/AnimGraphNode_BlendProfileLayeredBlend.cpp
+++ b/Engine/Plugins/Experimental/Animation/HierarchyTableAnimation/Source/UncookedOnly/Private/AnimGraphNode_BlendProfileLayeredBlend.cpp
@@ -59,26 +59,11 @@ void UAnimGraphNode_BlendProfileLayeredBlend::ValidateAnimNodeDuringCompilation(
 
 	const TObjectPtr<UHierarchyTable> MaskTable = Node.MaskTable;
 
-	if (MaskTable == nullptr)
+	if (MaskTable && !MaskTable->IsElementType<FHierarchyTable_ElementType_Mask>())
 	{
-		MessageLog.Error(*LOCTEXT("InvalidMask", "@@ has a null mask table.").ToString(), this, MaskTable);
+		MessageLog.Error(*LOCTEXT("InvalidMaskType", "@@ has a mask table that is not a mask type.").ToString(), this, MaskTable);
 		bCompilationError = true;
 	}
-	else if (MaskTable->IsElementType<FHierarchyTable_ElementType_Mask>())
-	{
-		MessageLog.Error(*LOCTEXT("InvalidMaskType", "@@ has a null mask table that is not a mask type.").ToString(), this, MaskTable);
-		bCompilationError = true;
-	}
-
-	if (bCompilationError)
-	{
-		return;
-	}
-
-	if (!Node.ArePerBoneBlendWeightsValid(ForSkeleton))
-	{
-		Node.RebuildPerBoneBlendWeights(ForSkeleton);
-	}
 }
 
 #undef LOCTEXT_NAMESPACE
