Begin transition of blend profiles being their own standalone asset instead of a sub-type of the hierarchy table asset.

diff --git a/Engine/Plugins/Experimental/Animation/HierarchyTable/Source/Editor/Private/HierachyTableEditorModule.cpp b/Engine/Plugins/Experimental/Animation/HierarchyTable/Source/Editor/Private/HierachyTableEditorModule.cpp
--- a/Engine/Plugins/Experimental/Animation/HierarchyTable/Source/Editor/Private/HierachyTableEditorModule.cpp
+++ b/Engine/Plugins/Experimental/Animation/HierarchyTable/Source/Editor/Private/HierachyTableEditorModule.cpp
@@ -6,6 +6,7 @@
 #include "HierarchyTableDefaultTypes.h"
 #include "DefaultHierarchyTableTypeHandler.h"
 #include "FloatColumn.h"
+#include "SHierarchyTable.h"
 
 void FHierarchyTableEditorModule::StartupModule()
 {
@@ -93,4 +94,9 @@ TArray<TSharedPtr<IHierarchyTableColumn>> FHierarchyTableEditorModule::GetElemen
 	return TArray<TSharedPtr<IHierarchyTableColumn>>();
 }
 
+TSharedRef<IHierarchyTable> FHierarchyTableEditorModule::CreateHierarchyTableWidget(const TObjectPtr<UHierarchyTable> HierarchyTable)
+{
+	return SNew(SHierarchyTable, HierarchyTable);
+}
+
 IMPLEMENT_MODULE(FHierarchyTableEditorModule, HierarchyTableEditor)
diff --git a/Engine/Plugins/Experimental/Animation/HierarchyTable/Source/Editor/Private/SHierarchyTable.h b/Engine/Plugins/Experimental/Animation/HierarchyTable/Source/Editor/Private/SHierarchyTable.h
--- a/Engine/Plugins/Experimental/Animation/HierarchyTable/Source/Editor/Private/SHierarchyTable.h
+++ b/Engine/Plugins/Experimental/Animation/HierarchyTable/Source/Editor/Private/SHierarchyTable.h
@@ -14,7 +14,7 @@ class SHierarchyTableRow;
 class UHierarchyTable;
 class UHierarchyTable_TableTypeHandler;
 
-class SHierarchyTable : public SCompoundWidget, public IHierarchyTable, public FSelfRegisteringEditorUndoClient, public FGCObject
+class SHierarchyTable : public IHierarchyTable, public FSelfRegisteringEditorUndoClient, public FGCObject
 {
 public:
 	friend SHierarchyTableRow;
diff --git a/Engine/Plugins/Experimental/Animation/HierarchyTable/Source/Editor/Public/HierarchyTableEditorModule.h b/Engine/Plugins/Experimental/Animation/HierarchyTable/Source/Editor/Public/HierarchyTableEditorModule.h
--- a/Engine/Plugins/Experimental/Animation/HierarchyTable/Source/Editor/Public/HierarchyTableEditorModule.h
+++ b/Engine/Plugins/Experimental/Animation/HierarchyTable/Source/Editor/Public/HierarchyTableEditorModule.h
@@ -9,6 +9,7 @@
 #include "UObject/ObjectKey.h"
 
 struct IHierarchyTableColumn;
+class IHierarchyTable;
 
 class FHierarchyTableEditorModule : public IModuleInterface
 {
@@ -17,6 +18,8 @@ public:
 	
 	virtual void ShutdownModule() override;
 
+	HIERARCHYTABLEEDITOR_API TSharedRef<IHierarchyTable> CreateHierarchyTableWidget(const TObjectPtr<UHierarchyTable> HierarchyTable);
+
 	HIERARCHYTABLEEDITOR_API void RegisterTableType(const UScriptStruct* TableType, const UClass* Handler);
 
 	HIERARCHYTABLEEDITOR_API void UnregisterTableType(const UScriptStruct* TableType);
diff --git a/Engine/Plugins/Experimental/Animation/HierarchyTable/Source/Editor/Public/IHierarchyTable.h b/Engine/Plugins/Experimental/Animation/HierarchyTable/Source/Editor/Public/IHierarchyTable.h
--- a/Engine/Plugins/Experimental/Animation/HierarchyTable/Source/Editor/Public/IHierarchyTable.h
+++ b/Engine/Plugins/Experimental/Animation/HierarchyTable/Source/Editor/Public/IHierarchyTable.h
@@ -3,8 +3,9 @@
 #pragma once
 
 #include "CoreMinimal.h"
+#include "Widgets/SCompoundWidget.h"
 
-class IHierarchyTable
+class IHierarchyTable : public SCompoundWidget
 {
 public:
 	virtual int32 GetSelectedEntryIndex() const = 0;
diff --git a/Engine/Plugins/Experimental/Animation/HierarchyTableAnimation/Source/Editor/HierarchyTableAnimationEditor.Build.cs b/Engine/Plugins/Experimental/Animation/HierarchyTableAnimation/Source/Editor/HierarchyTableAnimationEditor.Build.cs
--- a/Engine/Plugins/Experimental/Animation/HierarchyTableAnimation/Source/Editor/HierarchyTableAnimationEditor.Build.cs
+++ b/Engine/Plugins/Experimental/Animation/HierarchyTableAnimation/Source/Editor/HierarchyTableAnimationEditor.Build.cs
@@ -10,6 +10,7 @@ public class HierarchyTableAnimationEditor : ModuleRules
 		ShortName = "HTAnimEd";
 
 		PrivateDependencyModuleNames.AddAll(
+			"AssetDefinition",
 			"Core",
 			"CoreUObject",
 			"Engine", 
diff --git a/Engine/Plugins/Experimental/Animation/HierarchyTableAnimation/Source/Editor/Private/BlendProfileStandaloneAssetDefinition.cpp b/Engine/Plugins/Experimental/Animation/HierarchyTableAnimation/Source/Editor/Private/BlendProfileStandaloneAssetDefinition.cpp
new file mode 100644
--- /dev/null
+++ b/Engine/Plugins/Experimental/Animation/HierarchyTableAnimation/Source/Editor/Private/BlendProfileStandaloneAssetDefinition.cpp
@@ -0,0 +1,41 @@
+// Copyright Epic Games, Inc. All Rights Reserved.
+
+#include "BlendProfileStandaloneAssetDefinition.h"
+#include "BlendProfileStandalone.h"
+#include "BlendProfileStandaloneEditorToolkit.h"
+
+#define LOCTEXT_NAMESPACE "BlendProfileStandalone"
+
+FText UAssetDefinition_BlendProfileStandalone::GetAssetDisplayName() const
+{
+	return LOCTEXT("BlendProfileStandalone", "Blend Profile");
+}
+
+FLinearColor UAssetDefinition_BlendProfileStandalone::GetAssetColor() const
+{
+	return FLinearColor(FColor::Yellow);
+}
+
+TSoftClassPtr<UObject> UAssetDefinition_BlendProfileStandalone::GetAssetClass() const
+{
+	return UBlendProfileStandalone::StaticClass();
+}
+
+EAssetCommandResult UAssetDefinition_BlendProfileStandalone::OpenAssets(const FAssetOpenArgs& OpenArgs) const
+{
+	if (OpenArgs.OpenMethod == EAssetOpenMethod::Edit)
+	{
+		TArray<UObject*> Assets = OpenArgs.LoadObjects<UObject>();
+		MakeShared<FBlendProfileStandaloneEditorToolkit>()->InitEditor(Assets);
+	}
+
+	return EAssetCommandResult::Handled;
+}
+
+TConstArrayView<FAssetCategoryPath> UAssetDefinition_BlendProfileStandalone::GetAssetCategories() const
+{
+	static const auto Categories = { FAssetCategoryPath(EAssetCategoryPaths::Animation) };
+	return Categories;
+}
+
+#undef LOCTEXT_NAMESPACE
\ No newline at end of file
diff --git a/Engine/Plugins/Experimental/Animation/HierarchyTableAnimation/Source/Editor/Private/BlendProfileStandaloneAssetDefinition.h b/Engine/Plugins/Experimental/Animation/HierarchyTableAnimation/Source/Editor/Private/BlendProfileStandaloneAssetDefinition.h
new file mode 100644
--- /dev/null
+++ b/Engine/Plugins/Experimental/Animation/HierarchyTableAnimation/Source/Editor/Private/BlendProfileStandaloneAssetDefinition.h
@@ -0,0 +1,22 @@
+// Copyright Epic Games, Inc. All Rights Reserved.
+
+#pragma once
+
+#include "AssetDefinitionDefault.h"
+
+#include "BlendProfileStandaloneAssetDefinition.generated.h"
+
+UCLASS()
+class UAssetDefinition_BlendProfileStandalone : public UAssetDefinitionDefault
+{
+	GENERATED_BODY()
+
+public:
+	// Begin UAssetDefinition
+	virtual FText GetAssetDisplayName() const override;
+	virtual FLinearColor GetAssetColor() const override;
+	virtual TSoftClassPtr<UObject> GetAssetClass() const override;
+	virtual EAssetCommandResult OpenAssets(const FAssetOpenArgs& OpenArgs) const override;
+	virtual TConstArrayView<FAssetCategoryPath> GetAssetCategories() const override;
+	// End UAssetDefinition
+};
\ No newline at end of file
diff --git a/Engine/Plugins/Experimental/Animation/HierarchyTableAnimation/Source/Editor/Private/BlendProfileStandaloneEditorToolkit.cpp b/Engine/Plugins/Experimental/Animation/HierarchyTableAnimation/Source/Editor/Private/BlendProfileStandaloneEditorToolkit.cpp
new file mode 100644
--- /dev/null
+++ b/Engine/Plugins/Experimental/Animation/HierarchyTableAnimation/Source/Editor/Private/BlendProfileStandaloneEditorToolkit.cpp
@@ -0,0 +1,109 @@
+// Copyright Epic Games, Inc. All Rights Reserved.
+
+#include "BlendProfileStandaloneEditorToolkit.h"
+#include "HierarchyTableEditorModule.h"
+#include "Modules/ModuleManager.h"
+#include "PropertyEditorModule.h"
+#include "ToolMenus.h"
+#include "Widgets/Docking/SDockTab.h"
+
+#define LOCTEXT_NAMESPACE "BlendProfileStandaloneEditorToolkit"
+
+void FBlendProfileStandaloneEditorToolkit::InitEditor(const TArray<UObject*>& InObjects)
+{
+	BlendProfile = CastChecked<UBlendProfileStandalone>(InObjects[0]);
+
+	const TSharedRef<FTabManager::FLayout> Layout = FTabManager::NewLayout("BlendProfileStandaloneEditorToolkit")
+		->AddArea
+		(
+			FTabManager::NewPrimaryArea()->SetOrientation(Orient_Horizontal)
+			->Split
+			(
+				FTabManager::NewStack()
+				->SetSizeCoefficient(0.7f)
+				->AddTab("BlendProfileStandaloneEditorTableTab", ETabState::OpenedTab)
+			)
+			->Split
+			(
+				FTabManager::NewStack()
+				->SetSizeCoefficient(0.3f)
+				->AddTab("BlendProfileStandaloneEditorDetailsTab", ETabState::OpenedTab)
+			)
+		);
+
+	FAssetEditorToolkit::InitAssetEditor(EToolkitMode::Standalone, {}, "BlendProfileStandaloneEditor", Layout, true, true, InObjects);
+
+	ExtendToolbar();
+}
+
+void FBlendProfileStandaloneEditorToolkit::RegisterTabSpawners(const TSharedRef<class FTabManager>& InTabManager)
+{
+	FAssetEditorToolkit::RegisterTabSpawners(InTabManager);
+
+	WorkspaceMenuCategory = InTabManager->AddLocalWorkspaceMenuCategory(LOCTEXT("BlendProfileStandaloneEditor", "Blend Profile Editor"));
+
+	// Table view
+	{
+		FHierarchyTableEditorModule& HierarchyTableModule = FModuleManager::GetModuleChecked<FHierarchyTableEditorModule>("HierarchyTableEditor");
+
+		TSharedRef<IHierarchyTable> TableWidget = HierarchyTableModule.CreateHierarchyTableWidget(BlendProfile->Table);
+		HierarchyTableWidgetInterface = TableWidget;
+
+		InTabManager->RegisterTabSpawner("BlendProfileStandaloneEditorTableTab", FOnSpawnTab::CreateLambda([this, TableWidget](const FSpawnTabArgs&)
+			{
+				return SNew(SDockTab)
+					[
+						TableWidget
+					];
+			}))
+			.SetDisplayName(LOCTEXT("BlendProfile", "Blend Profile"))
+			.SetGroup(WorkspaceMenuCategory.ToSharedRef());
+	}
+
+	// Details panel
+	{
+		FPropertyEditorModule& PropertyEditorModule = FModuleManager::GetModuleChecked<FPropertyEditorModule>("PropertyEditor");
+		
+		FDetailsViewArgs DetailsViewArgs;
+		DetailsViewArgs.NameAreaSettings = FDetailsViewArgs::HideNameArea;
+		TSharedRef<IDetailsView> DetailsView = PropertyEditorModule.CreateDetailView(DetailsViewArgs);
+		DetailsView->SetObjects(TArray<UObject*>{ BlendProfile->Table });
+
+		InTabManager->RegisterTabSpawner("BlendProfileStandaloneEditorDetailsTab", FOnSpawnTab::CreateLambda([=](const FSpawnTabArgs&)
+			{
+				return SNew(SDockTab)
+					[
+						DetailsView
+					];
+			}))
+			.SetDisplayName(INVTEXT("Details"))
+			.SetGroup(WorkspaceMenuCategory.ToSharedRef());
+	}
+}
+
+void FBlendProfileStandaloneEditorToolkit::UnregisterTabSpawners(const TSharedRef<class FTabManager>& InTabManager)
+{
+	FAssetEditorToolkit::UnregisterTabSpawners(InTabManager);
+	InTabManager->UnregisterTabSpawner("BlendProfileStandaloneEditorTableTab");
+	InTabManager->UnregisterTabSpawner("BlendProfileStandaloneEditorDetailsTab");
+}
+
+void FBlendProfileStandaloneEditorToolkit::ExtendToolbar()
+{
+	FHierarchyTableEditorModule& HierarchyTableModule = FModuleManager::GetModuleChecked<FHierarchyTableEditorModule>("HierarchyTableEditor");
+	const TObjectPtr<UHierarchyTable_TableTypeHandler> Handler = HierarchyTableModule.CreateTableHandler(BlendProfile->Table);
+	if (!ensure(Handler))
+	{
+		return;
+	}
+
+	FToolMenuOwnerScoped OwnerScoped(this);
+
+	FName ParentName;
+	static const FName MenuName = GetToolMenuToolbarName(ParentName);
+
+	UToolMenu* ToolMenu = UToolMenus::Get()->ExtendMenu(MenuName);
+	Handler->ExtendToolbar(ToolMenu, *HierarchyTableWidgetInterface);
+}
+
+#undef LOCTEXT_NAMESPACE
\ No newline at end of file
diff --git a/Engine/Plugins/Experimental/Animation/HierarchyTableAnimation/Source/Editor/Private/BlendProfileStandaloneEditorToolkit.h b/Engine/Plugins/Experimental/Animation/HierarchyTableAnimation/Source/Editor/Private/BlendProfileStandaloneEditorToolkit.h
new file mode 100644
--- /dev/null
+++ b/Engine/Plugins/Experimental/Animation/HierarchyTableAnimation/Source/Editor/Private/BlendProfileStandaloneEditorToolkit.h
@@ -0,0 +1,33 @@
+// Copyright Epic Games, Inc. All Rights Reserved.
+
+#pragma once
+
+#include "CoreMinimal.h"
+#include "BlendProfileStandalone.h"
+#include "Toolkits/AssetEditorToolkit.h"
+
+class SHierarchyTable;
+class IHierarchyTable;
+
+class FBlendProfileStandaloneEditorToolkit : public FAssetEditorToolkit
+{
+public:
+	void InitEditor(const TArray<UObject*>& InObjects);
+
+	// Begin FAssetEditorToolkit
+	void RegisterTabSpawners(const TSharedRef<class FTabManager>& TabManager) override;
+	void UnregisterTabSpawners(const TSharedRef<class FTabManager>& TabManager) override;
+
+	FName GetToolkitFName() const override { return "BlendProfileEditor"; }
+	FText GetBaseToolkitName() const override { return INVTEXT("Blend Profile Editor"); }
+	FString GetWorldCentricTabPrefix() const override { return "Blend Profile "; }
+	FLinearColor GetWorldCentricTabColorScale() const override { return {}; }
+	// End FAssetEditorToolkit
+
+private:
+	void ExtendToolbar();
+
+	TObjectPtr<UBlendProfileStandalone> BlendProfile;
+
+	TSharedPtr<IHierarchyTable> HierarchyTableWidgetInterface;
+};
diff --git a/Engine/Plugins/Experimental/Animation/HierarchyTableAnimation/Source/Editor/Private/BlendProfileStandaloneFactory.cpp b/Engine/Plugins/Experimental/Animation/HierarchyTableAnimation/Source/Editor/Private/BlendProfileStandaloneFactory.cpp
new file mode 100644
--- /dev/null
+++ b/Engine/Plugins/Experimental/Animation/HierarchyTableAnimation/Source/Editor/Private/BlendProfileStandaloneFactory.cpp
@@ -0,0 +1,42 @@
+// Copyright Epic Games, Inc. All Rights Reserved.
+
+#include "BlendProfileStandaloneFactory.h"
+#include "BlendProfileStandalone.h"
+#include "Modules/ModuleManager.h"
+#include "SkeletonHierarchyTableType.h"
+#include "MaskProfile/HierarchyTableTypeMask.h"
+#include "HierarchyTableEditorModule.h"
+
+UBlendProfileStandaloneFactory::UBlendProfileStandaloneFactory()
+{
+	SupportedClass = UBlendProfileStandalone::StaticClass();
+	bCreateNew = true;
+}
+
+UObject* UBlendProfileStandaloneFactory::FactoryCreateNew(UClass* Class, UObject* InParent, FName Name, EObjectFlags Flags, UObject* Context, FFeedbackContext* Warn)
+{
+	TObjectPtr<UBlendProfileStandalone> BlendProfile = NewObject<UBlendProfileStandalone>(InParent, Class, Name, Flags, Context);
+	BlendProfile->Table = NewObject<UHierarchyTable>(BlendProfile);
+
+	// TODO: Steamline hierarchy table creation API
+
+	FInstancedStruct TableMetadata(FHierarchyTable_TableType_Skeleton::StaticStruct());
+	BlendProfile->Table->Initialize(TableMetadata, FHierarchyTable_ElementType_Mask::StaticStruct());
+
+	FHierarchyTableEditorModule& HierarchyTableModule = FModuleManager::GetModuleChecked<FHierarchyTableEditorModule>("HierarchyTableEditor");
+	const TObjectPtr<UHierarchyTable_TableTypeHandler> TableHandler = HierarchyTableModule.CreateTableHandler(BlendProfile->Table);
+	check(TableHandler);
+
+	TableHandler->FactoryConfigureProperties(TableMetadata);
+	BlendProfile->Table->Initialize(TableMetadata, FHierarchyTable_ElementType_Mask::StaticStruct());
+
+	TableHandler->SetHierarchyTable(BlendProfile->Table);
+	TableHandler->ConstructHierarchy();
+
+	return BlendProfile;
+}
+
+bool UBlendProfileStandaloneFactory::ConfigureProperties()
+{
+	return true;
+}
\ No newline at end of file
diff --git a/Engine/Plugins/Experimental/Animation/HierarchyTableAnimation/Source/Editor/Private/BlendProfileStandaloneFactory.h b/Engine/Plugins/Experimental/Animation/HierarchyTableAnimation/Source/Editor/Private/BlendProfileStandaloneFactory.h
new file mode 100644
--- /dev/null
+++ b/Engine/Plugins/Experimental/Animation/HierarchyTableAnimation/Source/Editor/Private/BlendProfileStandaloneFactory.h
@@ -0,0 +1,21 @@
+// Copyright Epic Games, Inc. All Rights Reserved.
+
+#pragma once
+
+#include "CoreMinimal.h"
+#include "Factories/Factory.h"
+
+#include "BlendProfileStandaloneFactory.generated.h"
+
+UCLASS()
+class UBlendProfileStandaloneFactory : public UFactory
+{
+	GENERATED_BODY()
+
+public:
+	UBlendProfileStandaloneFactory();
+
+	// UFactory interface
+	UObject* FactoryCreateNew(UClass* Class, UObject* InParent, FName Name, EObjectFlags Flags, UObject* Context, FFeedbackContext* Warn) override;
+	virtual bool ConfigureProperties() override;
+};
diff --git a/Engine/Plugins/Experimental/Animation/HierarchyTableAnimation/Source/Runtime/Public/BlendProfileStandalone.h b/Engine/Plugins/Experimental/Animation/HierarchyTableAnimation/Source/Runtime/Public/BlendProfileStandalone.h
new file mode 100644
--- /dev/null
+++ b/Engine/Plugins/Experimental/Animation/HierarchyTableAnimation/Source/Runtime/Public/BlendProfileStandalone.h
@@ -0,0 +1,18 @@
+// Copyright Epic Games, Inc. All Rights Reserved.
+
+#pragma once
+
+#include "CoreMinimal.h"
+#include "HierarchyTable.h"
+
+#include "BlendProfileStandalone.generated.h"
+
+UCLASS()
+class HIERARCHYTABLEANIMATIONRUNTIME_API UBlendProfileStandalone : public UObject
+{
+	GENERATED_BODY()
+
+public:
+	UPROPERTY(VisibleAnywhere, Category=Default)
+	TObjectPtr<UHierarchyTable> Table;
+};
\ No newline at end of file
