Allow pinnable blend profiles in profile blend node and refactor

diff --git a/Engine/Plugins/Experimental/Animation/HierarchyTableAnimation/Source/Runtime/Private/AnimNode_BlendProfileLayeredBlend.cpp b/Engine/Plugins/Experimental/Animation/HierarchyTableAnimation/Source/Runtime/Private/AnimNode_BlendProfileLayeredBlend.cpp
--- a/Engine/Plugins/Experimental/Animation/HierarchyTableAnimation/Source/Runtime/Private/AnimNode_BlendProfileLayeredBlend.cpp
+++ b/Engine/Plugins/Experimental/Animation/HierarchyTableAnimation/Source/Runtime/Private/AnimNode_BlendProfileLayeredBlend.cpp
@@ -10,8 +10,8 @@
 #include "MaskProfile/HierarchyTableTypeMask.h"
 #include "Animation/AttributeTypes.h"
 #include "Animation/IAttributeBlendOperator.h"
-#include "Hash/CityHash.h"
 #include "SkeletonHierarchyTableType.h"
+#include "HierarchyTableBlendProfile.h"
 
 #include UE_INLINE_GENERATED_CPP_BY_NAME(AnimNode_BlendProfileLayeredBlend)
 
@@ -24,76 +24,60 @@ void FAnimNode_BlendProfileLayeredBlend::Initialize_AnyThread(const FAnimationIn
 	BlendPose.Initialize(Context);
 }
 
-
-void FAnimNode_BlendProfileLayeredBlend::RebuildPerBoneBlendWeights(const USkeleton* InSkeleton)
-{
-	if (InSkeleton)
-	{
-		CreateMaskWeights(InSkeleton);
-
-		SkeletonGuid = InSkeleton->GetGuid();
-		VirtualBoneGuid = InSkeleton->GetVirtualBoneGuid();
-	}
-}
-
 bool FAnimNode_BlendProfileLayeredBlend::ArePerBoneBlendWeightsValid(const USkeleton* InSkeleton) const
 {
-	return (InSkeleton != nullptr && InSkeleton->GetGuid() == SkeletonGuid && InSkeleton->GetVirtualBoneGuid() == VirtualBoneGuid);
+	return (InSkeleton != nullptr && InSkeleton->GetGuid() == SkeletonGuid && InSkeleton->GetVirtualBoneGuid() == VirtualBoneGuid && CachedMaskTable == MaskTable);
 }
 
-void FAnimNode_BlendProfileLayeredBlend::UpdateCachedBoneData(const FBoneContainer& RequiredBones, const USkeleton* Skeleton)
+void FAnimNode_BlendProfileLayeredBlend::UpdateCachedBoneData(const FBoneContainer& RequiredBones, USkeleton* Skeleton)
 {
-	if (RequiredBones.GetSerialNumber() == RequiredBonesSerialNumber)
+	if (ArePerBoneBlendWeightsValid(Skeleton) && RequiredBonesSerialNumber == RequiredBones.GetSerialNumber())
 	{
 		return;
 	}
-	
-	// Update cached mask weights
+
+	if (Skeleton && MaskTable && MaskTable->IsTableType<FHierarchyTable_TableType_Skeleton>())
 	{
-		CachedCurveMaskWeights.Empty();
-		//CachedAttributeMaskWeights.Empty();
+		const FHierarchyTable_TableType_Skeleton& SkeletonMetadata = MaskTable->GetTableMetadata<FHierarchyTable_TableType_Skeleton>();
 
-		for (const FHierarchyTableEntryData& Entry : MaskTable->GetTableData())
+		if (SkeletonMetadata.Skeleton == Skeleton)
 		{
-			if (Entry.TablePayload.Get<FHierarchyTable_TablePayloadType_Skeleton>().EntryType == ESkeletonHierarchyTable_TablePayloadEntryType::Curve)
-			{
-				const float EntryWeight = Entry.GetValue<FHierarchyTable_ElementType_Mask>()->Value;
-				CachedCurveMaskWeights.Add<UE::Anim::FCurveElement>({ Entry.Identifier, EntryWeight });
-			}
-			/*
-			else if (Entry.EntryType == EHierarchyTableEntryType::Attribute)
-			{
-				const float EntryWeight = Entry.GetValue<FHierarchyTableType_Mask>()->Value;
-				CachedAttributeMaskWeights.Add(FNamedFloat(Entry.Identifier, EntryWeight));
-			}
-			*/
+			HierarchyTableBlendProfile = FHierarchyTableBlendProfile(MaskTable);
 		}
 	}
-
-	if (!ArePerBoneBlendWeightsValid(Skeleton))
+	else
 	{
-		RebuildPerBoneBlendWeights(Skeleton);
+		HierarchyTableBlendProfile = FHierarchyTableBlendProfile(Skeleton);
 	}
-
-	// build desired bone weights
-	const TArray<FBoneIndexType>& RequiredBoneIndices = RequiredBones.GetBoneIndicesArray();
-	const int32 NumRequiredBones = RequiredBoneIndices.Num();
-	DesiredBoneBlendWeights.SetNumZeroed(NumRequiredBones);
-	for (int32 RequiredBoneIndex = 0; RequiredBoneIndex < NumRequiredBones; RequiredBoneIndex++)
+	
+	// Update DesiredBoneBlendWeights
 	{
-		const int32 SkeletonBoneIndex = RequiredBones.GetSkeletonIndex(FCompactPoseBoneIndex(RequiredBoneIndex));
-		if (ensure(SkeletonBoneIndex != INDEX_NONE))
+		const TArray<FBoneIndexType>& RequiredBoneIndices = RequiredBones.GetBoneIndicesArray();
+		const int32 NumRequiredBones = RequiredBoneIndices.Num();
+		DesiredBoneBlendWeights.SetNumZeroed(NumRequiredBones);
+
+		for (int32 RequiredBoneIndex = 0; RequiredBoneIndex < NumRequiredBones; RequiredBoneIndex++)
 		{
-			DesiredBoneBlendWeights[RequiredBoneIndex] = PerBoneBlendWeights[SkeletonBoneIndex];
+			const int32 SkeletonBoneIndex = RequiredBones.GetSkeletonIndex(FCompactPoseBoneIndex(RequiredBoneIndex));
+			if (ensure(SkeletonBoneIndex != INDEX_NONE))
+			{
+				DesiredBoneBlendWeights[RequiredBoneIndex] = HierarchyTableBlendProfile.GetBoneBlendScale(SkeletonBoneIndex);
+			}
 		}
 	}
 
-	CurrentBoneBlendWeights.Reset(DesiredBoneBlendWeights.Num());
-	CurrentBoneBlendWeights.AddZeroed(DesiredBoneBlendWeights.Num());
+	// Update CurrentBoneBlendWeights
+	{
+		CurrentBoneBlendWeights.Reset(DesiredBoneBlendWeights.Num());
+		CurrentBoneBlendWeights.AddZeroed(DesiredBoneBlendWeights.Num());
 
-	//Reinitialize bone blend weights now that we have cleared them
-	UpdateDesiredBoneWeight();
+		//Reinitialize bone blend weights now that we have cleared them
+		UpdateDesiredBoneWeight();
+	}
 
+	SkeletonGuid = Skeleton ? Skeleton->GetGuid() : FGuid();
+	VirtualBoneGuid = Skeleton ? Skeleton->GetVirtualBoneGuid() : FGuid();
+	CachedMaskTable = MaskTable;
 	RequiredBonesSerialNumber = RequiredBones.GetSerialNumber();
 }
 
@@ -110,6 +94,7 @@ void FAnimNode_BlendProfileLayeredBlend::CacheBones_AnyThread(const FAnimationCa
 void FAnimNode_BlendProfileLayeredBlend::Update_AnyThread(const FAnimationUpdateContext& Context)
 {
 	DECLARE_SCOPE_HIERARCHICAL_COUNTER_ANIMNODE(Update_AnyThread)
+
 	bool RootMotionBlendPose = false;
 	float RootMotionWeight = 0.f;
 	const float RootMotionClearWeight = bBlendRootMotionBasedOnRootBone ? 0.f : 1.f;
@@ -186,6 +171,8 @@ void FAnimNode_BlendProfileLayeredBlend::Evaluate_AnyThread(FPoseContext& Output
 
 	// Blend curves
 	{
+		const UE::Anim::TNamedValueArray<FDefaultAllocator, UE::Anim::FCurveElement>& CachedCurveMaskWeights = HierarchyTableBlendProfile.GetCurveBlendWeights();
+
 		FBlendedCurve& OutCurve = OutputPoseData.GetCurve();
 		OutCurve.CopyFrom(BasePoseContext.Curve);
 
@@ -230,6 +217,8 @@ void FAnimNode_BlendProfileLayeredBlend::Evaluate_AnyThread(FPoseContext& Output
 
 	// Blend attributes
 	{
+		const TArray<FHierarchyTableBlendProfile::FNamedFloat>& CachedAttributeMaskWeights = HierarchyTableBlendProfile.GetAttributeBlendWeights();
+
 		using namespace UE::Anim;
 
 		FStackAttributeContainer& BaseAttributes = BasePoseContext.CustomAttributes;
@@ -240,8 +229,7 @@ void FAnimNode_BlendProfileLayeredBlend::Evaluate_AnyThread(FPoseContext& Output
 		// Premultiply masked blend attributes
 		const TArray<TWeakObjectPtr<UScriptStruct>>& BlendUniqueTypes = BlendAttributes.GetUniqueTypes();
 
-		//CachedAttributeMaskWeights.ForEachElement([&BlendUniqueTypes, &BlendAttributes](const FNamedFloat& AttributeMask)
-		for (const FNamedFloat& AttributeMask : CachedAttributeMaskWeights)
+		for (const FHierarchyTableBlendProfile::FNamedFloat& AttributeMask : CachedAttributeMaskWeights)
 		{
 			for (int32 BlendUniqueTypeIndex = 0; BlendUniqueTypeIndex < BlendUniqueTypes.Num(); ++BlendUniqueTypeIndex)
 			{
@@ -251,8 +239,7 @@ void FAnimNode_BlendProfileLayeredBlend::Evaluate_AnyThread(FPoseContext& Output
 				const TWeakObjectPtr<UScriptStruct> AttributeType = BlendUniqueTypes[BlendUniqueTypeIndex];
 				const IAttributeBlendOperator* Operator = UE::Anim::AttributeTypes::GetTypeOperator(AttributeType);
 
-				TWrappedAttribute<FAnimStackAllocator> DefaultData;
-				AttributeType->InitializeStruct(DefaultData.GetPtr<void>());
+				FWrappedAttribute DefaultValue(AttributeType.Get());
 
 				const int32 AttributeIndex = TypedAttributeIdentifiers.IndexOfByPredicate([&AttributeMask](const FAttributeId& InAttribute)
 					{
@@ -272,7 +259,7 @@ void FAnimNode_BlendProfileLayeredBlend::Evaluate_AnyThread(FPoseContext& Output
 						uint8* Value = BlendAttributes.Find(AttributeType.Get(), AttributeId);
 						check(Value != nullptr);
 
-						Operator->Interpolate(DefaultData.GetPtr<void>(), Value, AttributeMask.Value, Value);
+						Operator->Interpolate(DefaultValue.GetPtr<void>(), Value, AttributeMask.Value, Value);
 					}
 					break;
 				}
@@ -295,39 +282,6 @@ void FAnimNode_BlendProfileLayeredBlend::GatherDebugData(FNodeDebugData& DebugDa
 	BlendPose.GatherDebugData(DebugData.BranchFlow(BlendWeight));
 }
 
-void FAnimNode_BlendProfileLayeredBlend::CreateMaskWeights(const USkeleton* Skeleton)
-{
-	if (!Skeleton)
-	{
-		return;
-	}
-
-	const FReferenceSkeleton& RefSkeleton = Skeleton->GetReferenceSkeleton();
-
-	const int32 NumBones = RefSkeleton.GetNum();
-	PerBoneBlendWeights.Reset(NumBones);
-	PerBoneBlendWeights.AddZeroed(NumBones);
-
-	check(MaskTable->IsTableType<FHierarchyTable_TableType_Skeleton>());
-	check(MaskTable->IsElementType<FHierarchyTable_ElementType_Mask>());
-
-	const FHierarchyTable_TableType_Skeleton& SkeletonMetadata = MaskTable->GetTableMetadata<FHierarchyTable_TableType_Skeleton>();
-
-	for (int32 EntryIndex = 0; EntryIndex < MaskTable->GetTableData().Num(); ++EntryIndex)
-	{
-		const FHierarchyTableEntryData* const EntryData = MaskTable->GetTableEntry(EntryIndex);
-		if (EntryData->TablePayload.Get<FHierarchyTable_TablePayloadType_Skeleton>().EntryType != ESkeletonHierarchyTable_TablePayloadEntryType::Bone)
-		{
-			break;
-		}
-
-		const FHierarchyTable_ElementType_Mask* MaskValue = EntryData->GetValue<FHierarchyTable_ElementType_Mask>();
-		check(MaskValue);
-			
-		PerBoneBlendWeights[EntryIndex] = MaskValue->Value;
-	}
-}
-
 void FAnimNode_BlendProfileLayeredBlend::UpdateDesiredBoneWeight()
 {
 	check(CurrentBoneBlendWeights.Num() == DesiredBoneBlendWeights.Num());
diff --git a/Engine/Plugins/Experimental/Animation/HierarchyTableAnimation/Source/Runtime/Private/HierarchyTableBlendProfile.cpp b/Engine/Plugins/Experimental/Animation/HierarchyTableAnimation/Source/Runtime/Private/HierarchyTableBlendProfile.cpp
--- a/Engine/Plugins/Experimental/Animation/HierarchyTableAnimation/Source/Runtime/Private/HierarchyTableBlendProfile.cpp
+++ b/Engine/Plugins/Experimental/Animation/HierarchyTableAnimation/Source/Runtime/Private/HierarchyTableBlendProfile.cpp
@@ -32,6 +32,7 @@ FHierarchyTableBlendProfile::FHierarchyTableBlendProfile(TObjectPtr<UHierarchyTa
 	}
 	else if (InHierarchyTable->IsElementType<FHierarchyTable_ElementType_Mask>())
 	{
+		// Update bone weights
 		for (int32 BoneIndex = 0; BoneIndex < BoneCount; ++BoneIndex)
 		{
 			if (const FHierarchyTableEntryData* EntryData = InHierarchyTable->GetTableEntry(BoneIndex))
@@ -43,6 +44,26 @@ FHierarchyTableBlendProfile::FHierarchyTableBlendProfile(TObjectPtr<UHierarchyTa
 				BoneBlendWeights.Add(0.0f);
 			}
 		}
+
+		// Update curve and attribute weights
+		{
+			CurveBlendWeights.Empty();
+			AttributeBlendWeights.Empty();
+
+			for (const FHierarchyTableEntryData& Entry : InHierarchyTable->GetTableData())
+			{
+				if (Entry.TablePayload.Get<FHierarchyTable_TablePayloadType_Skeleton>().EntryType == ESkeletonHierarchyTable_TablePayloadEntryType::Curve)
+				{
+					const float EntryWeight = Entry.GetValue<FHierarchyTable_ElementType_Mask>()->Value;
+					CurveBlendWeights.Add<UE::Anim::FCurveElement>({ Entry.Identifier, EntryWeight });
+				}
+				else if (Entry.TablePayload.Get<FHierarchyTable_TablePayloadType_Skeleton>().EntryType == ESkeletonHierarchyTable_TablePayloadEntryType::Attribute)
+				{
+					const float EntryWeight = Entry.GetValue<FHierarchyTable_ElementType_Mask>()->Value;
+					AttributeBlendWeights.Add(FNamedFloat(Entry.Identifier, EntryWeight));
+				}
+			}
+		}
 	}
 	else
 	{
@@ -54,6 +75,21 @@ FHierarchyTableBlendProfile::FHierarchyTableBlendProfile(TObjectPtr<UHierarchyTa
 	}
 }
 
+FHierarchyTableBlendProfile::FHierarchyTableBlendProfile(const TObjectPtr<USkeleton> InSkeleton)
+{
+	if (InSkeleton)
+	{
+		Skeleton = InSkeleton;
+		const int32 BoneCount = Skeleton->GetReferenceSkeleton().GetNum();
+		BoneBlendWeights.SetNumZeroed(BoneCount);
+	}
+	else
+	{
+		Skeleton = nullptr;
+		BoneBlendWeights.SetNumZeroed(0);
+	}
+}
+
 float FHierarchyTableBlendProfile::GetBoneBlendScale(int32 InBoneIdx) const
 {
 	return BoneBlendWeights[InBoneIdx];
@@ -77,4 +113,14 @@ int32 FHierarchyTableBlendProfile::GetPerBoneInterpolationIndex(const FSkeletonP
 TObjectPtr<USkeleton> FHierarchyTableBlendProfile::GetSkeleton() const
 {
 	return Skeleton;
+}
+
+const UE::Anim::TNamedValueArray<FDefaultAllocator, UE::Anim::FCurveElement>& FHierarchyTableBlendProfile::GetCurveBlendWeights() const
+{
+	return CurveBlendWeights;
+}
+
+const TArray<FHierarchyTableBlendProfile::FNamedFloat>& FHierarchyTableBlendProfile::GetAttributeBlendWeights() const
+{
+	return AttributeBlendWeights;
 }
\ No newline at end of file
diff --git a/Engine/Plugins/Experimental/Animation/HierarchyTableAnimation/Source/Runtime/Public/AnimNode_BlendProfileLayeredBlend.h b/Engine/Plugins/Experimental/Animation/HierarchyTableAnimation/Source/Runtime/Public/AnimNode_BlendProfileLayeredBlend.h
--- a/Engine/Plugins/Experimental/Animation/HierarchyTableAnimation/Source/Runtime/Public/AnimNode_BlendProfileLayeredBlend.h
+++ b/Engine/Plugins/Experimental/Animation/HierarchyTableAnimation/Source/Runtime/Public/AnimNode_BlendProfileLayeredBlend.h
@@ -6,8 +6,10 @@
 #include "UObject/ObjectMacros.h"
 #include "Animation/AnimTypes.h"
 #include "Animation/AnimNodeBase.h"
+#include "HierarchyTableBlendProfile.h"
 #include "AnimNode_BlendProfileLayeredBlend.generated.h"
 
+
 class UHierarchyTable;
 enum class EHierarchyTableEntryType : uint8;
 
@@ -28,51 +30,29 @@ public:
 	/**
 	 * The mask hierarchy table to use to control layering of the pose, curves, and attributes
 	 */
-	UPROPERTY(EditAnywhere, BlueprintReadWrite, Category = Config)
+	UPROPERTY(EditAnywhere, BlueprintReadWrite, Category = Config, meta = (PinHiddenByDefault))
 	TObjectPtr<UHierarchyTable> MaskTable;
 
 protected:
-	// Per-bone weights for the skeleton. Serialized as these are only relative to the skeleton, but can potentially
-	// be regenerated at runtime if the GUIDs dont match
-	UPROPERTY()
-	TArray<float> PerBoneBlendWeights;
-
-	// transient data to handle weight and target weight
-	// this array changes based on required bones
-	TArray<float> DesiredBoneBlendWeights;
-	TArray<float> CurrentBoneBlendWeights;
-
-	UE::Anim::TNamedValueArray<FDefaultAllocator, UE::Anim::FCurveElement> CachedCurveMaskWeights;
-
-	struct FNamedFloat
-	{
-		FName Name;
-		float Value;
-
-		FNamedFloat(const FName InName, const float InValue)
-			: Name(InName)
-			, Value(InValue)
-		{
-		}
-	};
-	
-	TArray<FNamedFloat> CachedAttributeMaskWeights;
-
-	// Guids for skeleton used to determine whether the PerBoneBlendWeights need rebuilding
+	// Guids for skeleton used to determine whether the HierarchyTableBlendProfile need rebuilding
 	UPROPERTY()
 	FGuid SkeletonGuid;
 
-	// Guid for virtual bones used to determine whether the PerBoneBlendWeights need rebuilding
+	// Guid for virtual bones used to determine whether the HierarchyTableBlendProfile need rebuilding
 	UPROPERTY()
 	FGuid VirtualBoneGuid;
 
+	// Guid for mask table used to determine whether the HierarchyTableBlendProfile need rebuilding
+	UPROPERTY()
+	FGuid MaskTableGuid;
+
+	// Serial number of the required bones container
+	uint16 RequiredBonesSerialNumber = 0;
+
 	/** The weight of target pose */
 	UPROPERTY(EditAnywhere, BlueprintReadWrite, Category = Runtime, meta = (BlueprintCompilerGeneratedDefaults, PinShownByDefault))
 	float BlendWeight;
 
-	// Serial number of the required bones container
-	uint16 RequiredBonesSerialNumber;
-
 	/** Whether to incorporate the per-bone blend weight of the root bone when lending root motion */
 	UPROPERTY(EditAnywhere, Category = Config)
 	bool bBlendRootMotionBasedOnRootBone;
@@ -80,7 +60,6 @@ protected:
 public:
 	FAnimNode_BlendProfileLayeredBlend()
 		: BlendWeight(1.0f)
-		, RequiredBonesSerialNumber(0)
 		, bBlendRootMotionBasedOnRootBone(true)
 	{
 	}
@@ -93,25 +72,22 @@ public:
 	virtual void GatherDebugData(FNodeDebugData& DebugData) override;
 	// End of FAnimNode_Base interface
 
-	// Invalidate the cached per-bone blend weights from the skeleton
-	void InvalidatePerBoneBlendWeights() { RequiredBonesSerialNumber = 0; SkeletonGuid = FGuid(); VirtualBoneGuid = FGuid(); }
-
-	// Invalidates the cached bone data so it is recalculated the next time this node is updated
-	void InvalidateCachedBoneData() { RequiredBonesSerialNumber = 0; }
-
-	void CreateMaskWeights(const USkeleton* Skeleton);
-
 private:
-	// Rebuild cache per bone blend weights from the skeleton
-	void RebuildPerBoneBlendWeights(const USkeleton* InSkeleton);
-
 	// Check whether per-bone blend weights are valid according to the skeleton (GUID check)
 	bool ArePerBoneBlendWeightsValid(const USkeleton* InSkeleton) const;
 
 	// Update cached data if required
-	void UpdateCachedBoneData(const FBoneContainer& RequiredBones, const USkeleton* Skeleton);
+	void UpdateCachedBoneData(const FBoneContainer& RequiredBones, USkeleton* Skeleton);
 
 	friend class UAnimGraphNode_BlendProfileLayeredBlend;
 
 	void UpdateDesiredBoneWeight();
+
+	FHierarchyTableBlendProfile HierarchyTableBlendProfile;
+
+	TArray<float> DesiredBoneBlendWeights;
+	TArray<float> CurrentBoneBlendWeights;
+
+	UPROPERTY()
+	TObjectPtr<UHierarchyTable> CachedMaskTable;
 };
\ No newline at end of file
diff --git a/Engine/Plugins/Experimental/Animation/HierarchyTableAnimation/Source/Runtime/Public/HierarchyTableBlendProfile.h b/Engine/Plugins/Experimental/Animation/HierarchyTableAnimation/Source/Runtime/Public/HierarchyTableBlendProfile.h
--- a/Engine/Plugins/Experimental/Animation/HierarchyTableAnimation/Source/Runtime/Public/HierarchyTableBlendProfile.h
+++ b/Engine/Plugins/Experimental/Animation/HierarchyTableAnimation/Source/Runtime/Public/HierarchyTableBlendProfile.h
@@ -13,6 +13,8 @@ public:
 
 	FHierarchyTableBlendProfile(TObjectPtr<UHierarchyTable> InHierarchyTable);
 
+	FHierarchyTableBlendProfile(const TObjectPtr<USkeleton> InSkeleton);
+
 	// IBlendProfileInterface
 	virtual float GetBoneBlendScale(int32 InBoneIdx) const override;
 	virtual int32 GetNumBlendEntries() const override;
@@ -20,8 +22,27 @@ public:
 	virtual int32 GetPerBoneInterpolationIndex(const FSkeletonPoseBoneIndex InSkeletonBoneIndex, const USkeleton* TargetSkeleton, const IInterpolationIndexProvider::FPerBoneInterpolationData* Data) const override;
 	virtual TObjectPtr<USkeleton> GetSkeleton() const override;
 
+	struct FNamedFloat
+	{
+		FName Name;
+		float Value;
+
+		FNamedFloat(const FName InName, const float InValue)
+			: Name(InName)
+			, Value(InValue)
+		{
+		}
+	};
+
+	const UE::Anim::TNamedValueArray<FDefaultAllocator, UE::Anim::FCurveElement>& GetCurveBlendWeights() const;
+	const TArray<FNamedFloat>& GetAttributeBlendWeights() const;
+
 private:
 	TObjectPtr<USkeleton> Skeleton;
 
 	TArray<float> BoneBlendWeights;
+
+	UE::Anim::TNamedValueArray<FDefaultAllocator, UE::Anim::FCurveElement> CurveBlendWeights;
+
+	TArray<FNamedFloat> AttributeBlendWeights;
 };
\ No newline at end of file
