Add local/mesh space blending toggle to Profile Blend node

diff --git a/Engine/Plugins/Experimental/Animation/HierarchyTableAnimation/Source/Runtime/Private/AnimNode_BlendProfileLayeredBlend.cpp b/Engine/Plugins/Experimental/Animation/HierarchyTableAnimation/Source/Runtime/Private/AnimNode_BlendProfileLayeredBlend.cpp
--- a/Engine/Plugins/Experimental/Animation/HierarchyTableAnimation/Source/Runtime/Private/AnimNode_BlendProfileLayeredBlend.cpp
+++ b/Engine/Plugins/Experimental/Animation/HierarchyTableAnimation/Source/Runtime/Private/AnimNode_BlendProfileLayeredBlend.cpp
@@ -167,7 +167,20 @@ void FAnimNode_BlendProfileLayeredBlend::Evaluate_AnyThread(FPoseContext& Output
 	FAnimationPoseData OutputPoseData(Output);
 	
 	// Blend poses
-	FAnimationRuntime::BlendTwoPosesTogetherPerBone(BasePoseContext.Pose, TargetPose, CurrentBoneBlendWeights, OutputPoseData.GetPose());
+	{
+		if (bMeshSpaceRotationBlend)
+		{
+			FAnimationRuntime::ConvertPoseToMeshRotation(BasePoseContext.Pose);
+			FAnimationRuntime::ConvertPoseToMeshRotation(TargetPose);
+		}
+
+		FAnimationRuntime::BlendTwoPosesTogetherPerBone(BasePoseContext.Pose, TargetPose, CurrentBoneBlendWeights, OutputPoseData.GetPose());
+
+		if (bMeshSpaceRotationBlend)
+		{
+			FAnimationRuntime::ConvertMeshRotationPoseToLocalSpace(OutputPoseData.GetPose());
+		}
+	}
 
 	// Blend curves
 	{
diff --git a/Engine/Plugins/Experimental/Animation/HierarchyTableAnimation/Source/Runtime/Public/AnimNode_BlendProfileLayeredBlend.h b/Engine/Plugins/Experimental/Animation/HierarchyTableAnimation/Source/Runtime/Public/AnimNode_BlendProfileLayeredBlend.h
--- a/Engine/Plugins/Experimental/Animation/HierarchyTableAnimation/Source/Runtime/Public/AnimNode_BlendProfileLayeredBlend.h
+++ b/Engine/Plugins/Experimental/Animation/HierarchyTableAnimation/Source/Runtime/Public/AnimNode_BlendProfileLayeredBlend.h
@@ -33,6 +33,10 @@ public:
 	UPROPERTY(EditAnywhere, BlueprintReadWrite, Category = Config, meta = (PinHiddenByDefault))
 	TObjectPtr<UHierarchyTable> MaskTable;
 
+	/** Whether to blend bone rotations in mesh space or in local space */
+	UPROPERTY(EditAnywhere, BlueprintReadWrite, Category = Config)
+	bool bMeshSpaceRotationBlend;
+
 protected:
 	// Guids for skeleton used to determine whether the HierarchyTableBlendProfile need rebuilding
 	UPROPERTY()
