Remove legacy time profile type from hierarchy table

diff --git a/Engine/Plugins/Experimental/Animation/HierarchyTableAnimation/Source/Editor/Private/HierachyTableAnimationEditorModule.cpp b/Engine/Plugins/Experimental/Animation/HierarchyTableAnimation/Source/Editor/Private/HierachyTableAnimationEditorModule.cpp
--- a/Engine/Plugins/Experimental/Animation/HierarchyTableAnimation/Source/Editor/Private/HierachyTableAnimationEditorModule.cpp
+++ b/Engine/Plugins/Experimental/Animation/HierarchyTableAnimation/Source/Editor/Private/HierachyTableAnimationEditorModule.cpp
@@ -1,14 +1,12 @@
 // Copyright Epic Games, Inc. All Rights Reserved.
 
 #include "HierarchyTableAnimationEditorModule.h"
-#include "TimeProfile/HierarchyTableTypeTime.h"
 #include "MaskProfile/HierarchyTableTypeMask.h"
 #include "Modules/ModuleManager.h"
 #include "HierarchyTableEditorModule.h"
 #include "SkeletonHierarchyTableType.h"
 #include "SkeletonHierarchyTableTypeHandler.h"
 #include "MaskProfile/MaskProfileColumn.h"
-#include "TimeProfile/TimeProfileColumn.h"
 #include "SkeletonHierarchyTableTypeDetailsCustomization.h"
 #include "PropertyEditorModule.h"
 #include "PersonaModule.h"
@@ -27,15 +25,6 @@ void FHierarchyTableAnimationEditorModule::StartupModule()
 			MakeShared<FHierarchyTableColumn_Mask>(),
 		});
 
-	BuiltinElementTypes.Add(FHierarchyTable_ElementType_Time::StaticStruct());
-	HierarchyTableModule.RegisterElementTypeEditorColumns(FHierarchyTable_ElementType_Time::StaticStruct(),
-		{
-			MakeShared<FHierarchyTableColumn_TimeProfile_StartTime>(),
-			MakeShared<FHierarchyTableColumn_TimeProfile_EndTime>(),
-			MakeShared<FHierarchyTableColumn_TimeProfile_TimeFactor>(),
-			MakeShared<FHierarchyTableColumn_TimeProfile_Preview>(),
-		});
-
 	FPropertyEditorModule& PropertyModule = FModuleManager::LoadModuleChecked<FPropertyEditorModule>("PropertyEditor");
 	PropertyModule.RegisterCustomPropertyTypeLayout(FHierarchyTable_TableType_Skeleton::StaticStruct()->GetFName(), FOnGetPropertyTypeCustomizationInstance::CreateStatic(&FHierarchyTableSkeletonTableTypeDetailsCustomization::MakeInstance));
 
diff --git a/Engine/Plugins/Experimental/Animation/HierarchyTableAnimation/Source/Editor/Private/TimeProfile/TimeProfileColumn.cpp b/Engine/Plugins/Experimental/Animation/HierarchyTableAnimation/Source/Editor/Private/TimeProfile/TimeProfileColumn.cpp
deleted file mode 100644
diff --git a/Engine/Plugins/Experimental/Animation/HierarchyTableAnimation/Source/Editor/Private/TimeProfile/TimeProfileColumn.h b/Engine/Plugins/Experimental/Animation/HierarchyTableAnimation/Source/Editor/Private/TimeProfile/TimeProfileColumn.h
deleted file mode 100644
diff --git a/Engine/Plugins/Experimental/Animation/HierarchyTableAnimation/Source/Runtime/Public/TimeProfile/HierarchyTableTypeTime.h b/Engine/Plugins/Experimental/Animation/HierarchyTableAnimation/Source/Runtime/Public/TimeProfile/HierarchyTableTypeTime.h
deleted file mode 100644
