Add missing struct member initialization
- Fixes errors logged during startup of the EngineTest project
- Also let the uproject file update so it stops prompting on every run

diff --git a/Engine/Plugins/Experimental/Animation/HierarchyTableAnimation/Source/Runtime/Public/AnimNode_BlendProfileLayeredBlend.h b/Engine/Plugins/Experimental/Animation/HierarchyTableAnimation/Source/Runtime/Public/AnimNode_BlendProfileLayeredBlend.h
--- a/Engine/Plugins/Experimental/Animation/HierarchyTableAnimation/Source/Runtime/Public/AnimNode_BlendProfileLayeredBlend.h
+++ b/Engine/Plugins/Experimental/Animation/HierarchyTableAnimation/Source/Runtime/Public/AnimNode_BlendProfileLayeredBlend.h
@@ -35,7 +35,7 @@ public:
 
 	/** Whether to blend bone rotations in mesh space or in local space */
 	UPROPERTY(EditAnywhere, BlueprintReadWrite, Category = Config)
-	bool bMeshSpaceRotationBlend;
+	bool bMeshSpaceRotationBlend = false;
 
 protected:
 	// Guids for skeleton used to determine whether the HierarchyTableBlendProfile need rebuilding
@@ -94,4 +94,4 @@ private:
 
 	UPROPERTY()
 	TObjectPtr<UHierarchyTable> CachedMaskTable;
-};
\ No newline at end of file
+};
diff --git a/Engine/Source/Runtime/Experimental/Iris/Core/Public/Iris/ReplicationSystem/NetTokenStore.h b/Engine/Source/Runtime/Experimental/Iris/Core/Public/Iris/ReplicationSystem/NetTokenStore.h
--- a/Engine/Source/Runtime/Experimental/Iris/Core/Public/Iris/ReplicationSystem/NetTokenStore.h
+++ b/Engine/Source/Runtime/Experimental/Iris/Core/Public/Iris/ReplicationSystem/NetTokenStore.h
@@ -64,7 +64,7 @@ struct FNetTokenStoreTypeIdPair
 	UPROPERTY()
 	FString StoreTypeName;
 	UPROPERTY()
-	uint32 TypeID;
+	uint32 TypeID = 0;
 	bool operator<(const FNetTokenStoreTypeIdPair& Other) const
 	{
 		return TypeID < Other.TypeID;
diff --git a/Engine/Source/Runtime/MovieScene/Private/Tests/MovieSceneTestObjects.h b/Engine/Source/Runtime/MovieScene/Private/Tests/MovieSceneTestObjects.h
--- a/Engine/Source/Runtime/MovieScene/Private/Tests/MovieSceneTestObjects.h
+++ b/Engine/Source/Runtime/MovieScene/Private/Tests/MovieSceneTestObjects.h
@@ -171,7 +171,7 @@ struct FTestMovieSceneStruct
 	float Second = 0.f;
 
 	UPROPERTY(EditAnywhere, BlueprintReadWrite, Interp, Category="Test")
-	FVector Vector;
+	FVector Vector = FVector::ZeroVector;
 };
 
 UCLASS(MinimalAPI, BlueprintType)
