Add custom curve blending options ot ProfileBlend node (powered by [REDACTED_PERSON])
Expose mesh space bool as a pin
Add profiler tag to UpdateCachedBoneData

diff --git a/Engine/Plugins/Experimental/Animation/HierarchyTableAnimation/Source/Runtime/Private/AnimNode_BlendProfileLayeredBlend.cpp b/Engine/Plugins/Experimental/Animation/HierarchyTableAnimation/Source/Runtime/Private/AnimNode_BlendProfileLayeredBlend.cpp
--- a/Engine/Plugins/Experimental/Animation/HierarchyTableAnimation/Source/Runtime/Private/AnimNode_BlendProfileLayeredBlend.cpp
+++ b/Engine/Plugins/Experimental/Animation/HierarchyTableAnimation/Source/Runtime/Private/AnimNode_BlendProfileLayeredBlend.cpp
@@ -31,11 +31,13 @@ bool FAnimNode_BlendProfileLayeredBlend::ArePerBoneBlendWeightsValid(const USkel
 
 void FAnimNode_BlendProfileLayeredBlend::UpdateCachedBoneData(const FBoneContainer& RequiredBones, USkeleton* Skeleton)
 {
+
 	if (ArePerBoneBlendWeightsValid(Skeleton) && RequiredBonesSerialNumber == RequiredBones.GetSerialNumber())
 	{
 		return;
 	}
 
+	QUICK_SCOPE_CYCLE_COUNTER(STAT_BlendProfileLayeredBlend_UpdateCachedBoneData);
 	if (Skeleton && MaskTable && MaskTable->IsTableType<FHierarchyTable_TableType_Skeleton>())
 	{
 		const FHierarchyTable_TableType_Skeleton& SkeletonMetadata = MaskTable->GetTableMetadata<FHierarchyTable_TableType_Skeleton>();
@@ -183,6 +185,7 @@ void FAnimNode_BlendProfileLayeredBlend::Evaluate_AnyThread(FPoseContext& Output
 	}
 
 	// Blend curves
+	if (!bCustomCurveBlending)
 	{
 		const UE::Anim::TNamedValueArray<FDefaultAllocator, UE::Anim::FCurveElement>& CachedCurveMaskWeights = HierarchyTableBlendProfile.GetCurveBlendWeights();
 
@@ -227,6 +230,22 @@ void FAnimNode_BlendProfileLayeredBlend::Evaluate_AnyThread(FPoseContext& Output
 					});
 		}
 	}
+	else
+	{
+		float TargetPoseMaxWeight = 0.0f;
+
+		for (const FCompactPoseBoneIndex BoneIndex : TargetPose.ForEachBoneIndex())
+		{
+			const float BoneBlendWeight = FMath::Clamp(CurrentBoneBlendWeights[BoneIndex.GetInt()], 0.f, 1.f);
+			TargetPoseMaxWeight = FMath::Max(TargetPoseMaxWeight, BoneBlendWeight);
+		}
+
+		const TArray<const FBlendedCurve*> SourceCurves = { &BasePoseContext.Curve, &TargetCurve };
+		const TArray<float> SourceWeights = { BlendWeight, TargetPoseMaxWeight };
+
+		FBlendedCurve& OutCurve = OutputPoseData.GetCurve();
+		BlendCurves(SourceCurves, SourceWeights, OutCurve, CurveBlendingOption);
+	}
 
 	// Blend attributes
 	{
diff --git a/Engine/Plugins/Experimental/Animation/HierarchyTableAnimation/Source/Runtime/Public/AnimNode_BlendProfileLayeredBlend.h b/Engine/Plugins/Experimental/Animation/HierarchyTableAnimation/Source/Runtime/Public/AnimNode_BlendProfileLayeredBlend.h
--- a/Engine/Plugins/Experimental/Animation/HierarchyTableAnimation/Source/Runtime/Public/AnimNode_BlendProfileLayeredBlend.h
+++ b/Engine/Plugins/Experimental/Animation/HierarchyTableAnimation/Source/Runtime/Public/AnimNode_BlendProfileLayeredBlend.h
@@ -34,9 +34,15 @@ public:
 	TObjectPtr<UHierarchyTable> MaskTable;
 
 	/** Whether to blend bone rotations in mesh space or in local space */
-	UPROPERTY(EditAnywhere, BlueprintReadWrite, Category = Config)
+	UPROPERTY(EditAnywhere, BlueprintReadWrite, Category = Config, meta = (PinHiddenByDefault))
 	bool bMeshSpaceRotationBlend = false;
 
+	UPROPERTY(EditAnywhere, BlueprintReadWrite, Category = Config)
+	bool bCustomCurveBlending;
+
+	UPROPERTY(EditAnywhere, BlueprintReadWrite, Category = Config, meta = (EditCondition = "bCustomCurveBlending"))
+	TEnumAsByte<enum ECurveBlendOption::Type> CurveBlendingOption;
+
 protected:
 	// Guids for skeleton used to determine whether the HierarchyTableBlendProfile need rebuilding
 	UPROPERTY()
@@ -63,7 +69,9 @@ protected:
 
 public:
 	FAnimNode_BlendProfileLayeredBlend()
-		: BlendWeight(1.0f)
+		: bCustomCurveBlending(false)
+		, CurveBlendingOption(ECurveBlendOption::Override)
+		, BlendWeight(1.0f)
 		, bBlendRootMotionBasedOnRootBone(true)
 	{
 	}
@@ -94,4 +102,4 @@ private:
 
 	UPROPERTY()
 	TObjectPtr<UHierarchyTable> CachedMaskTable;
-};
+};
\ No newline at end of file
