Temp fix for Mac build error while I sort it out

diff --git a/Engine/Plugins/Experimental/Animation/HierarchyTableAnimation/Source/Runtime/Private/AnimNode_BlendProfileLayeredBlend.cpp b/Engine/Plugins/Experimental/Animation/HierarchyTableAnimation/Source/Runtime/Private/AnimNode_BlendProfileLayeredBlend.cpp
--- a/Engine/Plugins/Experimental/Animation/HierarchyTableAnimation/Source/Runtime/Private/AnimNode_BlendProfileLayeredBlend.cpp
+++ b/Engine/Plugins/Experimental/Animation/HierarchyTableAnimation/Source/Runtime/Private/AnimNode_BlendProfileLayeredBlend.cpp
@@ -185,7 +185,7 @@ void FAnimNode_BlendProfileLayeredBlend::Evaluate_AnyThread(FPoseContext& Output
 	}
 
 	// Blend curves
-	if (!bCustomCurveBlending)
+	//if (!bCustomCurveBlending)
 	{
 		const UE::Anim::TNamedValueArray<FDefaultAllocator, UE::Anim::FCurveElement>& CachedCurveMaskWeights = HierarchyTableBlendProfile.GetCurveBlendWeights();
 
@@ -230,7 +230,7 @@ void FAnimNode_BlendProfileLayeredBlend::Evaluate_AnyThread(FPoseContext& Output
 					});
 		}
 	}
-	else
+	/*else
 	{
 		float TargetPoseMaxWeight = 0.0f;
 
@@ -245,7 +245,7 @@ void FAnimNode_BlendProfileLayeredBlend::Evaluate_AnyThread(FPoseContext& Output
 
 		FBlendedCurve& OutCurve = OutputPoseData.GetCurve();
 		BlendCurves(SourceCurves, SourceWeights, OutCurve, CurveBlendingOption);
-	}
+	}*/
 
 	// Blend attributes
 	{
