Fix for blend attributes ensure hit in blend profile layered blend node

diff --git a/Engine/Plugins/Experimental/Animation/HierarchyTableAnimation/Source/Runtime/Private/AnimNode_BlendProfileLayeredBlend.cpp b/Engine/Plugins/Experimental/Animation/HierarchyTableAnimation/Source/Runtime/Private/AnimNode_BlendProfileLayeredBlend.cpp
--- a/Engine/Plugins/Experimental/Animation/HierarchyTableAnimation/Source/Runtime/Private/AnimNode_BlendProfileLayeredBlend.cpp
+++ b/Engine/Plugins/Experimental/Animation/HierarchyTableAnimation/Source/Runtime/Private/AnimNode_BlendProfileLayeredBlend.cpp
@@ -302,7 +302,9 @@ void FAnimNode_BlendProfileLayeredBlend::Evaluate_AnyThread(FPoseContext& Output
 		const TArray<const float> SourceWeights { 1.0f - BlendWeight, BlendWeight };
 		static const TArray<const int32> WeightView { 0, 1 };
 
-		UE::Anim::Attributes::BlendAttributes(SourceAttributes, SourceWeights, WeightView, BaseAttributes);
+		FStackAttributeContainer BlendedAttributes;
+		UE::Anim::Attributes::BlendAttributes(SourceAttributes, SourceWeights, WeightView, BlendedAttributes);
+		OutAttributes.MoveFrom(BlendedAttributes);
 	}
 }
 
