Temp fix for crash from array OOB in hierarchy tabel blend profile until actual skeleton invalidation code is submitted

diff --git a/Engine/Plugins/Experimental/Animation/HierarchyTableAnimation/Source/Runtime/Private/AnimNode_BlendProfileLayeredBlend.cpp b/Engine/Plugins/Experimental/Animation/HierarchyTableAnimation/Source/Runtime/Private/AnimNode_BlendProfileLayeredBlend.cpp
--- a/Engine/Plugins/Experimental/Animation/HierarchyTableAnimation/Source/Runtime/Private/AnimNode_BlendProfileLayeredBlend.cpp
+++ b/Engine/Plugins/Experimental/Animation/HierarchyTableAnimation/Source/Runtime/Private/AnimNode_BlendProfileLayeredBlend.cpp
@@ -61,7 +61,7 @@ void FAnimNode_BlendProfileLayeredBlend::UpdateCachedBoneData(const FBoneContain
 		for (int32 RequiredBoneIndex = 0; RequiredBoneIndex < NumRequiredBones; RequiredBoneIndex++)
 		{
 			const int32 SkeletonBoneIndex = RequiredBones.GetSkeletonIndex(FCompactPoseBoneIndex(RequiredBoneIndex));
-			if (ensure(SkeletonBoneIndex != INDEX_NONE))
+			if (HierarchyTableBlendProfile.IsValidBoneIndex(SkeletonBoneIndex))
 			{
 				DesiredBoneBlendWeights[RequiredBoneIndex] = HierarchyTableBlendProfile.GetBoneBlendScale(SkeletonBoneIndex);
 			}
diff --git a/Engine/Plugins/Experimental/Animation/HierarchyTableAnimation/Source/Runtime/Private/HierarchyTableBlendProfile.cpp b/Engine/Plugins/Experimental/Animation/HierarchyTableAnimation/Source/Runtime/Private/HierarchyTableBlendProfile.cpp
--- a/Engine/Plugins/Experimental/Animation/HierarchyTableAnimation/Source/Runtime/Private/HierarchyTableBlendProfile.cpp
+++ b/Engine/Plugins/Experimental/Animation/HierarchyTableAnimation/Source/Runtime/Private/HierarchyTableBlendProfile.cpp
@@ -143,4 +143,9 @@ void FHierarchyTableBlendProfile::ConstructBlendProfile(const TObjectPtr<UBlendP
 
 		OutBlendProfile->SetBoneBlendScale(BoneIndex, BlendWeight, false /* bRecurse */, true /* bCreate */);
 	}
+}
+
+bool FHierarchyTableBlendProfile::IsValidBoneIndex(const int32 BoneIndex) const
+{
+	return BoneBlendWeights.IsValidIndex(BoneIndex);
 }
\ No newline at end of file
diff --git a/Engine/Plugins/Experimental/Animation/HierarchyTableAnimation/Source/Runtime/Public/HierarchyTableBlendProfile.h b/Engine/Plugins/Experimental/Animation/HierarchyTableAnimation/Source/Runtime/Public/HierarchyTableBlendProfile.h
--- a/Engine/Plugins/Experimental/Animation/HierarchyTableAnimation/Source/Runtime/Public/HierarchyTableBlendProfile.h
+++ b/Engine/Plugins/Experimental/Animation/HierarchyTableAnimation/Source/Runtime/Public/HierarchyTableBlendProfile.h
@@ -40,6 +40,8 @@ public:
 
 	void ConstructBlendProfile(const TObjectPtr<UBlendProfile> OutBlendProfile) const;
 
+	bool IsValidBoneIndex(const int32 BoneIndex) const;
+
 private:
 	TObjectPtr<USkeleton> Skeleton;
 
