Fix facial curve blending in blend layer trait due to curves in some cases iincorrectly lerping with 0.0

diff --git a/Engine/Plugins/Experimental/AnimNextAnimGraph/Source/AnimNextAnimGraph/Private/EvaluationVM/Tasks/BlendKeyframesPerBone.cpp b/Engine/Plugins/Experimental/AnimNextAnimGraph/Source/AnimNextAnimGraph/Private/EvaluationVM/Tasks/BlendKeyframesPerBone.cpp
--- a/Engine/Plugins/Experimental/AnimNextAnimGraph/Source/AnimNextAnimGraph/Private/EvaluationVM/Tasks/BlendKeyframesPerBone.cpp
+++ b/Engine/Plugins/Experimental/AnimNextAnimGraph/Source/AnimNextAnimGraph/Private/EvaluationVM/Tasks/BlendKeyframesPerBone.cpp
@@ -320,8 +320,11 @@ void FAnimNextBlendKeyframePerBoneWithScaleTask::Execute(UE::AnimNext::FEvaluati
 				TargetCurve,
 				[this](UE::Anim::FCurveElement& InOutBaseElement, const UE::Anim::FCurveElement& InBlendElement, UE::Anim::ENamedValueUnionFlags InFlags)
 				{
-					InOutBaseElement.Value = FMath::Lerp(InOutBaseElement.Value, InBlendElement.Value, ScaleFactor);
-					InOutBaseElement.Flags |= InBlendElement.Flags;
+					if (InFlags == UE::Anim::ENamedValueUnionFlags::BothArgsValid || InFlags == UE::Anim::ENamedValueUnionFlags::ValidArg1)
+					{
+						InOutBaseElement.Value = FMath::Lerp(InOutBaseElement.Value, InBlendElement.Value, ScaleFactor);
+						InOutBaseElement.Flags |= InBlendElement.Flags;
+					}
 				});
 		}
 	}
diff --git a/Engine/Plugins/Experimental/Animation/HierarchyTableAnimation/Source/Runtime/Private/AnimNode_BlendProfileLayeredBlend.cpp b/Engine/Plugins/Experimental/Animation/HierarchyTableAnimation/Source/Runtime/Private/AnimNode_BlendProfileLayeredBlend.cpp
--- a/Engine/Plugins/Experimental/Animation/HierarchyTableAnimation/Source/Runtime/Private/AnimNode_BlendProfileLayeredBlend.cpp
+++ b/Engine/Plugins/Experimental/Animation/HierarchyTableAnimation/Source/Runtime/Private/AnimNode_BlendProfileLayeredBlend.cpp
@@ -225,8 +225,11 @@ void FAnimNode_BlendProfileLayeredBlend::Evaluate_AnyThread(FPoseContext& Output
 				TargetCurve,
 				[this](UE::Anim::FCurveElement& InOutBaseElement, const UE::Anim::FCurveElement& InBlendElement, UE::Anim::ENamedValueUnionFlags InFlags)
 					{
-						InOutBaseElement.Value = FMath::Lerp(InOutBaseElement.Value, InBlendElement.Value, BlendWeight);
-						InOutBaseElement.Flags |= InBlendElement.Flags; // Should this only apply when the weight is relevant?
+						if (InFlags == UE::Anim::ENamedValueUnionFlags::BothArgsValid || InFlags == UE::Anim::ENamedValueUnionFlags::ValidArg1)
+						{
+							InOutBaseElement.Value = FMath::Lerp(InOutBaseElement.Value, InBlendElement.Value, BlendWeight);
+							InOutBaseElement.Flags |= InBlendElement.Flags; // Should this only apply when the weight is relevant?
+						}
 					});
 		}
 	}
