Fix blending of attributes in blend layer trait and matching code in ABP node
* Clean up Profile Blend node code a bit
* Include description of what's happening for future reference.

diff --git a/Engine/Plugins/Experimental/AnimNextAnimGraph/Source/AnimNextAnimGraph/Private/EvaluationVM/Tasks/BlendKeyframesPerBone.cpp b/Engine/Plugins/Experimental/AnimNextAnimGraph/Source/AnimNextAnimGraph/Private/EvaluationVM/Tasks/BlendKeyframesPerBone.cpp
--- a/Engine/Plugins/Experimental/AnimNextAnimGraph/Source/AnimNextAnimGraph/Private/EvaluationVM/Tasks/BlendKeyframesPerBone.cpp
+++ b/Engine/Plugins/Experimental/AnimNextAnimGraph/Source/AnimNextAnimGraph/Private/EvaluationVM/Tasks/BlendKeyframesPerBone.cpp
@@ -199,7 +199,7 @@ void FAnimNextBlendAddKeyframePerBoneWithScaleTask::Execute(UE::AnimNext::FEvalu
 	VM.PushValue(KEYFRAME_STACK_NAME, MoveTemp(KeyframeB));
 }
 
-FAnimNextBlendKeyframePerBoneWithScaleTask FAnimNextBlendKeyframePerBoneWithScaleTask::Make(TSharedPtr<IInterpolationIndexProvider> BlendProfile, const USkeleton* Skeleton, const TArray<float>& BoneMaskWeights, const UE::Anim::TNamedValueArray<FDefaultAllocator, UE::Anim::FCurveElement>& CurveMaskWeights, const TArray<FNamedFloat>& AttributeMaskWeights, float ScaleFactor)
+FAnimNextBlendKeyframePerBoneWithScaleTask FAnimNextBlendKeyframePerBoneWithScaleTask::Make(TSharedPtr<IInterpolationIndexProvider> BlendProfile, const USkeleton* Skeleton, const TArray<float>& BoneMaskWeights, const UE::Anim::TNamedValueArray<FDefaultAllocator, UE::Anim::FCurveElement>& CurveMaskWeights, const TArray<FMaskedAttributeWeight>& AttributeMaskWeights, float ScaleFactor)
 {
 	FAnimNextBlendKeyframePerBoneWithScaleTask Task;
 	Task.BlendProfile = BlendProfile;
@@ -336,54 +336,107 @@ void FAnimNextBlendKeyframePerBoneWithScaleTask::Execute(UE::AnimNext::FEvaluati
 		FStackAttributeContainer& BaseAttributes = KeyframeA->Attributes;
 		FStackAttributeContainer& BlendAttributes = KeyframeB->Attributes;
 
-		// Premultiply masked blend attributes
-		const TArray<TWeakObjectPtr<UScriptStruct>>& BlendUniqueTypes = BlendAttributes.GetUniqueTypes();
+		FStackAttributeContainer OutputAttributes;
+
+		// Attributes are to be masked out according to the mask weights in AttributeMaskWeights, if an attribute has no mask weight set then it
+		// inherits the weight of whatever bone it is attached to. Below are possible configurations that we need to account for:
+
+		// Root 0.0						Root set to 0.0 therefore a RootMotionDelta attribute will also be masked out without having to set an explicit entry in AttributeMaskWeights
+
+		// Root 0.0						RootMotionDelta is set to 1.0 in AttributeMaskWeights despite the parent bone being masked out.
+		//  \ RootMotionDelta 1.0
+
+		// Root 1.0						RootMotionDelta is being masked out in AttributeMaskWeights despite the parent bone being kept.
+		//  \ RootMotionDelta 0.0
+
+
+		// Below is a table of the possible permutations of base/blend attributes being present/absent along with the possible mask values
+		// k denotes some value in the range (0, 1) exclusive.
+		// - denotes an absent attribute
+
+		// Base | Blend | Weight | Output
+		// ------------------------------
+		// a    | b     | 1.0    | b
+		// a    | b     | k      | lerp(a, b, k)
+		// a    | b     | 0.0    | a
+		// - - - - - - - - - - - - - - -
+		// a    | -     | 1.0    | a
+		// a    | -     | k      | a
+		// a    | -     | 0.0    | a
+		// - - - - - - - - - - - - - - -
+		// -    | b     | 1.0    | b
+		// -    | b     | k      | lerp(default, b, k)
+		// -    | b     | 0.0    | -
+
+
+		// 1. Blend attributes according to the bone blend weights,
+		// i.e. an attribute's weight is determined by the weight of its attached bone.
+		UE::Anim::Attributes::BlendAttributesPerBone(BaseAttributes, BlendAttributes, BoneMaskWeights, OutputAttributes);
+
 
-		//CachedAttributeMaskWeights.ForEachElement([&BlendUniqueTypes, &BlendAttributes](const FNamedFloat& AttributeMask)
-		for (const FNamedFloat& AttributeMask : AttributeMaskWeights)
+		// 2. For each attribute that has a custom weight, i.e. one's that shouldn't be weighted by its attached bone, go and correct the blended value.
+		for (const FMaskedAttributeWeight& MaskedAttribute : AttributeMaskWeights)
 		{
-			for (int32 BlendUniqueTypeIndex = 0; BlendUniqueTypeIndex < BlendUniqueTypes.Num(); ++BlendUniqueTypeIndex)
-			{
-				const TArray<FAttributeId>& TypedAttributeIdentifiers = BlendAttributes.GetKeys(BlendUniqueTypeIndex);
-				const TArray<UE::Anim::TWrappedAttribute<FAnimStackAllocator>>& TypedAttributeValues = BlendAttributes.GetValues(BlendUniqueTypeIndex);
+			const TArray<TWeakObjectPtr<UScriptStruct>>& UniqueTypes = OutputAttributes.GetUniqueTypes();
 
-				const TWeakObjectPtr<UScriptStruct> AttributeType = BlendUniqueTypes[BlendUniqueTypeIndex];
+			for (int32 UniqueTypeIndex = 0; UniqueTypeIndex < UniqueTypes.Num(); ++UniqueTypeIndex)
+			{
+				const TArray<FAttributeId>& AttributeKeys = OutputAttributes.GetKeys(UniqueTypeIndex);
+				const TArray<TWrappedAttribute<FAnimStackAllocator>>& AttributeValues = OutputAttributes.GetValues(UniqueTypeIndex);
+				
+				const TWeakObjectPtr<UScriptStruct> AttributeType = UniqueTypes[UniqueTypeIndex];
 				const IAttributeBlendOperator* Operator = UE::Anim::AttributeTypes::GetTypeOperator(AttributeType);
 
-				TWrappedAttribute<FAnimStackAllocator> DefaultData;
+				TWrappedAttribute<FAnimStackAllocator> DefaultData(AttributeType.Get());
 				AttributeType->InitializeStruct(DefaultData.GetPtr<void>());
 
-				const int32 AttributeIndex = TypedAttributeIdentifiers.IndexOfByPredicate([&AttributeMask](const FAttributeId& InAttribute)
-					{
-						return InAttribute.GetName() == AttributeMask.Name;
-					});
-
-				if (AttributeIndex != INDEX_NONE)
+				uint8* OutputData = OutputAttributes.Find(AttributeType.Get(), MaskedAttribute.Attribute);
+				if (OutputData)
 				{
-					const FAttributeId& AttributeId = TypedAttributeIdentifiers[AttributeIndex];
+					uint8* BaseData = BaseAttributes.Find(AttributeType.Get(), MaskedAttribute.Attribute);
+					uint8* BlendData = BlendAttributes.Find(AttributeType.Get(), MaskedAttribute.Attribute);
 
-					if (AttributeMask.Value == 0.0f)
+					if (BaseData && BlendData)
 					{
-						BlendAttributes.Remove(AttributeType.Get(), AttributeId);
+						// Base | Blend | Weight | Output
+						// ------------------------------
+						// a    | b     | 1.0    | b
+						// a    | b     | k      | lerp(a, b, k)
+						// a    | b     | 0.0    | a
+
+						Operator->Interpolate(BaseData, BlendData, MaskedAttribute.Weight, OutputData);
 					}
-					else
+					else if (BaseData && !BlendData)
 					{
-						uint8* Value = BlendAttributes.Find(AttributeType.Get(), AttributeId);
-						check(Value != nullptr);
+						// Base | Blend | Weight | Output
+						// ------------------------------
+						// a    | -     | 1.0    | a
+						// a    | -     | k      | a
+						// a    | -     | 0.0    | a
 
-						Operator->Interpolate(DefaultData.GetPtr<void>(), Value, AttributeMask.Value, Value);
+						*OutputData = *BaseData;
 					}
-					break;
-				}
+					else if (!BaseData && BlendData)
+					{
+						// Base | Blend | Weight | Output
+						// ------------------------------
+						// -    | b     | 1.0    | b
+						// -    | b     | k      | lerp(default, b, k)
+						// -    | b     | 0.0    | -
+
+						if (MaskedAttribute.Weight != 0.0f)
+						{
+							Operator->Interpolate(DefaultData.GetPtr<void>(), BlendData, MaskedAttribute.Weight, OutputData);
+						}
+						else
+						{
+							OutputAttributes.Remove(AttributeType.Get(), MaskedAttribute.Attribute);
+						}
+					}
+				}				
 			}
 		}
 
-		const TArray<const UE::Anim::FStackAttributeContainer> SourceAttributes{ BaseAttributes, BlendAttributes };
-		const TArray<const float> SourceWeights{ 1.0f - ScaleFactor, ScaleFactor };
-		static const TArray<const int32> WeightView{ 0, 1 };
-
-		FStackAttributeContainer OutputAttributes;
-		UE::Anim::Attributes::BlendAttributes(SourceAttributes, SourceWeights, WeightView, OutputAttributes);
 		BaseAttributes.MoveFrom(OutputAttributes);
 	}
 
diff --git a/Engine/Plugins/Experimental/AnimNextAnimGraph/Source/AnimNextAnimGraph/Private/Traits/BlendLayer.cpp b/Engine/Plugins/Experimental/AnimNextAnimGraph/Source/AnimNextAnimGraph/Private/Traits/BlendLayer.cpp
--- a/Engine/Plugins/Experimental/AnimNextAnimGraph/Source/AnimNextAnimGraph/Private/Traits/BlendLayer.cpp
+++ b/Engine/Plugins/Experimental/AnimNextAnimGraph/Source/AnimNextAnimGraph/Private/Traits/BlendLayer.cpp
@@ -144,8 +144,13 @@ namespace UE::AnimNext
 				{
 					if (TableEntry.GetMetadata<FHierarchyTable_TablePayloadType_Skeleton>().EntryType == ESkeletonHierarchyTable_TablePayloadEntryType::Attribute)
 					{
+						const FHierarchyTableEntryData* ParentEntry = HierarchyTable->GetTableEntry(TableEntry.Parent);
+						check(ParentEntry && ParentEntry->GetMetadata<FHierarchyTable_TablePayloadType_Skeleton>().EntryType == ESkeletonHierarchyTable_TablePayloadEntryType::Bone);
+
+						UE::Anim::FAttributeId Attribute(TableEntry.Identifier, TableEntry.Parent, TEXT("bone"));
+
 						const float EntryWeight = TableEntry.GetValue<FHierarchyTable_ElementType_Mask>()->Value;
-						InstanceData->AttributeMaskWeights.Add(FAnimNextBlendKeyframePerBoneWithScaleTask::FNamedFloat(TableEntry.Identifier, EntryWeight));
+						InstanceData->AttributeMaskWeights.Add(FAnimNextBlendKeyframePerBoneWithScaleTask::FMaskedAttributeWeight(Attribute, EntryWeight));
 					}
 				}
 			}
diff --git a/Engine/Plugins/Experimental/AnimNextAnimGraph/Source/AnimNextAnimGraph/Private/Traits/BlendLayer.h b/Engine/Plugins/Experimental/AnimNextAnimGraph/Source/AnimNextAnimGraph/Private/Traits/BlendLayer.h
--- a/Engine/Plugins/Experimental/AnimNextAnimGraph/Source/AnimNextAnimGraph/Private/Traits/BlendLayer.h
+++ b/Engine/Plugins/Experimental/AnimNextAnimGraph/Source/AnimNextAnimGraph/Private/Traits/BlendLayer.h
@@ -70,7 +70,7 @@ namespace UE::AnimNext
 
 			UE::Anim::TNamedValueArray<FDefaultAllocator, UE::Anim::FCurveElement> CurveMaskWeights;
 
-			TArray<FAnimNextBlendKeyframePerBoneWithScaleTask::FNamedFloat> AttributeMaskWeights;
+			TArray<FAnimNextBlendKeyframePerBoneWithScaleTask::FMaskedAttributeWeight> AttributeMaskWeights;
 		};
 
 
diff --git a/Engine/Plugins/Experimental/AnimNextAnimGraph/Source/AnimNextAnimGraph/Public/EvaluationVM/Tasks/BlendKeyframesPerBone.h b/Engine/Plugins/Experimental/AnimNextAnimGraph/Source/AnimNextAnimGraph/Public/EvaluationVM/Tasks/BlendKeyframesPerBone.h
--- a/Engine/Plugins/Experimental/AnimNextAnimGraph/Source/AnimNextAnimGraph/Public/EvaluationVM/Tasks/BlendKeyframesPerBone.h
+++ b/Engine/Plugins/Experimental/AnimNextAnimGraph/Source/AnimNextAnimGraph/Public/EvaluationVM/Tasks/BlendKeyframesPerBone.h
@@ -7,6 +7,7 @@
 #include "CoreMinimal.h"
 #include "EvaluationVM/EvaluationTask.h"
 #include "EvaluationVM/Tasks/BlendKeyframes.h"
+#include "Animation/AttributesContainer.h"
 
 #include "BlendKeyframesPerBone.generated.h"
 
@@ -95,19 +96,19 @@ struct ANIMNEXTANIMGRAPH_API FAnimNextBlendKeyframePerBoneWithScaleTask : public
 
 	DECLARE_ANIM_EVALUATION_TASK(FAnimNextBlendKeyframePerBoneWithScaleTask)
 
-	struct FNamedFloat
+	struct FMaskedAttributeWeight
 	{
-		FName Name;
-		float Value;
+		UE::Anim::FAttributeId Attribute;
+		float Weight;
 
-		FNamedFloat(const FName InName, const float InValue)
-			: Name(InName)
-			, Value(InValue)
+		FMaskedAttributeWeight(const UE::Anim::FAttributeId InAttribute, const float InWeight)
+			: Attribute(InAttribute)
+			, Weight(InWeight)
 		{
 		}
 	};
 
-	static FAnimNextBlendKeyframePerBoneWithScaleTask Make(TSharedPtr<IInterpolationIndexProvider> BlendProfile, const USkeleton* Skeleton, const TArray<float>& BoneMaskWeights, const UE::Anim::TNamedValueArray<FDefaultAllocator, UE::Anim::FCurveElement>& CurveMaskWeights, const TArray<FNamedFloat>& AttributeMaskWeights, float ScaleFactor);
+	static FAnimNextBlendKeyframePerBoneWithScaleTask Make(TSharedPtr<IInterpolationIndexProvider> BlendProfile, const USkeleton* Skeleton, const TArray<float>& BoneMaskWeights, const UE::Anim::TNamedValueArray<FDefaultAllocator, UE::Anim::FCurveElement>& CurveMaskWeights, const TArray<FMaskedAttributeWeight>& AttributeMaskWeights, float ScaleFactor);
 
 	// Task entry point
 	virtual void Execute(UE::AnimNext::FEvaluationVM& VM) const override;
@@ -117,7 +118,7 @@ private:
 
 	UE::Anim::TNamedValueArray<FDefaultAllocator, UE::Anim::FCurveElement> CurveMaskWeights;
 
-	TArray<FNamedFloat> AttributeMaskWeights;
+	TArray<FMaskedAttributeWeight> AttributeMaskWeights;
 
 	TSharedPtr<IInterpolationIndexProvider> BlendProfile = nullptr;
 
diff --git a/Engine/Plugins/Experimental/Animation/HierarchyTableAnimation/Source/Runtime/Private/AnimNode_BlendProfileLayeredBlend.cpp b/Engine/Plugins/Experimental/Animation/HierarchyTableAnimation/Source/Runtime/Private/AnimNode_BlendProfileLayeredBlend.cpp
--- a/Engine/Plugins/Experimental/Animation/HierarchyTableAnimation/Source/Runtime/Private/AnimNode_BlendProfileLayeredBlend.cpp
+++ b/Engine/Plugins/Experimental/Animation/HierarchyTableAnimation/Source/Runtime/Private/AnimNode_BlendProfileLayeredBlend.cpp
@@ -143,44 +143,34 @@ void FAnimNode_BlendProfileLayeredBlend::Evaluate_AnyThread(FPoseContext& Output
 	DECLARE_SCOPE_HIERARCHICAL_COUNTER_ANIMNODE(Evaluate_AnyThread)
 	ANIM_MT_SCOPE_CYCLE_COUNTER(BlendPosesInGraph, !IsInGameThread());
 
-	FPoseContext BasePoseContext(Output);
-	BasePose.Evaluate(BasePoseContext);
+	if (!FAnimWeight::IsRelevant(BlendWeight))
+	{
+		BasePose.Evaluate(Output);
+		return;
+	}
 
-	FCompactPose TargetPose;
-	FBlendedCurve TargetCurve;
-	UE::Anim::FStackAttributeContainer TargetAttributes;
+	FPoseContext BasePoseContext(Output);
+	FPoseContext BlendPoseContext(Output);
 
-	if (FAnimWeight::IsRelevant(BlendWeight))
-	{
-		FPoseContext CurrentPoseContext(Output);
-		BlendPose.Evaluate(CurrentPoseContext);
+	BasePose.Evaluate(BasePoseContext);
+	BlendPose.Evaluate(BlendPoseContext);
 
-		TargetPose.MoveBonesFrom(CurrentPoseContext.Pose);
-		TargetCurve.MoveFrom(CurrentPoseContext.Curve);
-		TargetAttributes.MoveFrom(CurrentPoseContext.CustomAttributes);
-	}
-	else
-	{
-		TargetPose.ResetToRefPose(BasePoseContext.Pose.GetBoneContainer());
-		TargetCurve.InitFrom(Output.Curve);
-		TargetAttributes.CopyFrom(Output.CustomAttributes);
-	}
+	FAnimationPoseData BasePoseData(BasePoseContext);
+	FAnimationPoseData BlendPoseData(BlendPoseContext);
 
-	FAnimationPoseData OutputPoseData(Output);
-	
 	// Blend poses
 	{
 		if (bMeshSpaceRotationBlend)
 		{
-			FAnimationRuntime::ConvertPoseToMeshRotation(BasePoseContext.Pose);
-			FAnimationRuntime::ConvertPoseToMeshRotation(TargetPose);
+			FAnimationRuntime::ConvertPoseToMeshRotation(BasePoseData.GetPose());
+			FAnimationRuntime::ConvertPoseToMeshRotation(BlendPoseData.GetPose());
 		}
 
-		FAnimationRuntime::BlendTwoPosesTogetherPerBone(BasePoseContext.Pose, TargetPose, CurrentBoneBlendWeights, OutputPoseData.GetPose());
+		FAnimationRuntime::BlendTwoPosesTogetherPerBone(BasePoseData.GetPose(), BlendPoseData.GetPose(), CurrentBoneBlendWeights, Output.Pose);
 
 		if (bMeshSpaceRotationBlend)
 		{
-			FAnimationRuntime::ConvertMeshRotationPoseToLocalSpace(OutputPoseData.GetPose());
+			FAnimationRuntime::ConvertMeshRotationPoseToLocalSpace(Output.Pose);
 		}
 	}
 
@@ -189,8 +179,7 @@ void FAnimNode_BlendProfileLayeredBlend::Evaluate_AnyThread(FPoseContext& Output
 	{
 		const UE::Anim::TNamedValueArray<FDefaultAllocator, UE::Anim::FCurveElement>& CachedCurveMaskWeights = HierarchyTableBlendProfile.GetCurveBlendWeights();
 
-		FBlendedCurve& OutCurve = OutputPoseData.GetCurve();
-		OutCurve.CopyFrom(BasePoseContext.Curve);
+		Output.Curve.CopyFrom(BasePoseData.GetCurve());
 
 		if (FAnimWeight::IsRelevant(BlendWeight))
 		{
@@ -198,7 +187,7 @@ void FAnimNode_BlendProfileLayeredBlend::Evaluate_AnyThread(FPoseContext& Output
 
 			// Multiply per-curve blend weights by matching blend pose curves
 			UE::Anim::FNamedValueArrayUtils::Intersection(
-				TargetCurve,
+				BlendPoseData.GetCurve(),
 				CachedCurveMaskWeights,
 				[this, &FilteredCurves](const UE::Anim::FCurveElement& InBlendElement, const UE::Anim::FCurveElement& InMaskElement) mutable
 					{
@@ -206,11 +195,11 @@ void FAnimNode_BlendProfileLayeredBlend::Evaluate_AnyThread(FPoseContext& Output
 					});
 
 			// Override blend curve values with premultipled curves
-			TargetCurve.Combine(FilteredCurves);
+			BlendPoseData.GetCurve().Combine(FilteredCurves);
 
 			// Remove curves that have been filtered by the mask, curves with no mask value defined remain, even with 0.0 value
 			UE::Anim::FNamedValueArrayUtils::RemoveByPredicate(
-				TargetCurve,
+				BlendPoseData.GetCurve(),
 				CachedCurveMaskWeights,
 				[](const UE::Anim::FCurveElement& InBaseElement, const UE::Anim::FCurveElement& InMaskElement)
 				{	
@@ -221,8 +210,8 @@ void FAnimNode_BlendProfileLayeredBlend::Evaluate_AnyThread(FPoseContext& Output
 
 			// Combine base and filtered pre-multiplied blend curves
 			UE::Anim::FNamedValueArrayUtils::Union(
-				OutCurve,
-				TargetCurve,
+				Output.Curve,
+				BlendPoseData.GetCurve(),
 				[this](UE::Anim::FCurveElement& InOutBaseElement, const UE::Anim::FCurveElement& InBlendElement, UE::Anim::ENamedValueUnionFlags InFlags)
 					{
 						if (InFlags == UE::Anim::ENamedValueUnionFlags::BothArgsValid || InFlags == UE::Anim::ENamedValueUnionFlags::ValidArg1)
@@ -237,77 +226,124 @@ void FAnimNode_BlendProfileLayeredBlend::Evaluate_AnyThread(FPoseContext& Output
 	{
 		float TargetPoseMaxWeight = 0.0f;
 
-		for (const FCompactPoseBoneIndex BoneIndex : TargetPose.ForEachBoneIndex())
+		for (const FCompactPoseBoneIndex BoneIndex : BlendPoseData.GetPose().ForEachBoneIndex())
 		{
 			const float BoneBlendWeight = FMath::Clamp(CurrentBoneBlendWeights[BoneIndex.GetInt()], 0.f, 1.f);
 			TargetPoseMaxWeight = FMath::Max(TargetPoseMaxWeight, BoneBlendWeight);
 		}
 
-		const TArray<const FBlendedCurve*> SourceCurves = { &BasePoseContext.Curve, &TargetCurve };
+		const TArray<const FBlendedCurve*> SourceCurves = { &BasePoseData.GetCurve(), &BlendPoseData.GetCurve() };
 		const TArray<float> SourceWeights = { BlendWeight, TargetPoseMaxWeight };
 
-		FBlendedCurve& OutCurve = OutputPoseData.GetCurve();
-		BlendCurves(SourceCurves, SourceWeights, OutCurve, CurveBlendingOption);
+		BlendCurves(SourceCurves, SourceWeights, Output.Curve, CurveBlendingOption);
 	}
 
 	// Blend attributes
 	{
-		const TArray<FHierarchyTableBlendProfile::FNamedFloat>& CachedAttributeMaskWeights = HierarchyTableBlendProfile.GetAttributeBlendWeights();
-
 		using namespace UE::Anim;
 
-		FStackAttributeContainer& BaseAttributes = BasePoseContext.CustomAttributes;
-		FStackAttributeContainer& BlendAttributes = TargetAttributes;
-		FStackAttributeContainer& OutAttributes = OutputPoseData.GetAttributes();
-		OutAttributes.CopyFrom(BaseAttributes);
+		FStackAttributeContainer OutputAttributes;
+
+		// Attributes are to be masked out according to the mask weights in AttributeMaskWeights, if an attribute has no mask weight set then it
+		// inherits the weight of whatever bone it is attached to. Below are possible configurations that we need to account for:
+
+		// Root 0.0						Root set to 0.0 therefore a RootMotionDelta attribute will also be masked out without having to set an explicit entry in AttributeMaskWeights
+
+		// Root 0.0						RootMotionDelta is set to 1.0 in AttributeMaskWeights despite the parent bone being masked out.
+		//  \ RootMotionDelta 1.0
+
+		// Root 1.0						RootMotionDelta is being masked out in AttributeMaskWeights despite the parent bone being kept.
+		//  \ RootMotionDelta 0.0
 
-		// Premultiply masked blend attributes
-		const TArray<TWeakObjectPtr<UScriptStruct>>& BlendUniqueTypes = BlendAttributes.GetUniqueTypes();
 
-		for (const FHierarchyTableBlendProfile::FNamedFloat& AttributeMask : CachedAttributeMaskWeights)
+		// Below is a table of the possible permutations of base/blend attributes being present/absent along with the possible mask values
+		// k denotes some value in the range (0, 1) exclusive.
+		// - denotes an absent attribute
+
+		// Base | Blend | Weight | Output
+		// ------------------------------
+		// a    | b     | 1.0    | b
+		// a    | b     | k      | lerp(a, b, k)
+		// a    | b     | 0.0    | a
+		// - - - - - - - - - - - - - - -
+		// a    | -     | 1.0    | a
+		// a    | -     | k      | a
+		// a    | -     | 0.0    | a
+		// - - - - - - - - - - - - - - -
+		// -    | b     | 1.0    | b
+		// -    | b     | k      | lerp(default, b, k)
+		// -    | b     | 0.0    | -
+
+		// 1. Blend attributes according to the bone blend weights,
+		// i.e. an attribute's weight is determined by the weight of its attached bone.
+		UE::Anim::Attributes::BlendAttributesPerBone(BasePoseData.GetAttributes(), BlendPoseData.GetAttributes(), CurrentBoneBlendWeights, OutputAttributes);
+
+		// 2. For each attribute that has a custom weight, i.e. one's that shouldn't be weighted by its attached bone, go and correct the blended value.
+		const TArray<FHierarchyTableBlendProfile::FMaskedAttributeWeight>& CachedAttributeMaskWeights = HierarchyTableBlendProfile.GetAttributeBlendWeights();
+		for (const FHierarchyTableBlendProfile::FMaskedAttributeWeight& MaskedAttribute : CachedAttributeMaskWeights)
 		{
-			for (int32 BlendUniqueTypeIndex = 0; BlendUniqueTypeIndex < BlendUniqueTypes.Num(); ++BlendUniqueTypeIndex)
-			{
-				const TArray<FAttributeId>& TypedAttributeIdentifiers = BlendAttributes.GetKeys(BlendUniqueTypeIndex);
-				const TArray<UE::Anim::TWrappedAttribute<FAnimStackAllocator>>& TypedAttributeValues = BlendAttributes.GetValues(BlendUniqueTypeIndex);
+			const TArray<TWeakObjectPtr<UScriptStruct>>& UniqueTypes = OutputAttributes.GetUniqueTypes();
 
-				const TWeakObjectPtr<UScriptStruct> AttributeType = BlendUniqueTypes[BlendUniqueTypeIndex];
+			for (int32 UniqueTypeIndex = 0; UniqueTypeIndex < UniqueTypes.Num(); ++UniqueTypeIndex)
+			{
+				const TArray<FAttributeId>& AttributeKeys = OutputAttributes.GetKeys(UniqueTypeIndex);
+				const TArray<TWrappedAttribute<FAnimStackAllocator>>& AttributeValues = OutputAttributes.GetValues(UniqueTypeIndex);
+				
+				const TWeakObjectPtr<UScriptStruct> AttributeType = UniqueTypes[UniqueTypeIndex];
 				const IAttributeBlendOperator* Operator = UE::Anim::AttributeTypes::GetTypeOperator(AttributeType);
 
-				FWrappedAttribute DefaultValue(AttributeType.Get());
+				TWrappedAttribute<FAnimStackAllocator> DefaultData(AttributeType.Get());
+				AttributeType->InitializeStruct(DefaultData.GetPtr<void>());
 
-				const int32 AttributeIndex = TypedAttributeIdentifiers.IndexOfByPredicate([&AttributeMask](const FAttributeId& InAttribute)
-					{
-						return InAttribute.GetName() == AttributeMask.Name;
-					});
-
-				if (AttributeIndex != INDEX_NONE)
+				uint8* OutputData = OutputAttributes.Find(AttributeType.Get(), MaskedAttribute.Attribute);
+				if (OutputData)
 				{
-					const FAttributeId& AttributeId = TypedAttributeIdentifiers[AttributeIndex];
+					uint8* BaseData = BasePoseData.GetAttributes().Find(AttributeType.Get(), MaskedAttribute.Attribute);
+					uint8* BlendData = BlendPoseData.GetAttributes().Find(AttributeType.Get(), MaskedAttribute.Attribute);
+
+					if (BaseData && BlendData)
+					{
+						// Base | Blend | Weight | Output
+						// ------------------------------
+						// a    | b     | 1.0    | b
+						// a    | b     | k      | lerp(a, b, k)
+						// a    | b     | 0.0    | a
 
-					if (AttributeMask.Value == 0.0f)
+						Operator->Interpolate(BaseData, BlendData, MaskedAttribute.Weight, OutputData);
+					}
+					else if (BaseData && !BlendData)
 					{
-						BlendAttributes.Remove(AttributeType.Get(), AttributeId);
+						// Base | Blend | Weight | Output
+						// ------------------------------
+						// a    | -     | 1.0    | a
+						// a    | -     | k      | a
+						// a    | -     | 0.0    | a
+
+						*OutputData = *BaseData;
 					}
-					else
+					else if (!BaseData && BlendData)
 					{
-						uint8* Value = BlendAttributes.Find(AttributeType.Get(), AttributeId);
-						check(Value != nullptr);
+						// Base | Blend | Weight | Output
+						// ------------------------------
+						// -    | b     | 1.0    | b
+						// -    | b     | k      | lerp(default, b, k)
+						// -    | b     | 0.0    | -
 
-						Operator->Interpolate(DefaultValue.GetPtr<void>(), Value, AttributeMask.Value, Value);
+						if (MaskedAttribute.Weight != 0.0f)
+						{
+							Operator->Interpolate(DefaultData.GetPtr<void>(), BlendData, MaskedAttribute.Weight, OutputData);
+							
+						}
+						else
+						{
+							OutputAttributes.Remove(AttributeType.Get(), MaskedAttribute.Attribute);
+						}
 					}
-					break;
-				}
+				}				
 			}
 		}
 
-		const TArray<const UE::Anim::FStackAttributeContainer> SourceAttributes { BaseAttributes, BlendAttributes };
-		const TArray<const float> SourceWeights { 1.0f - BlendWeight, BlendWeight };
-		static const TArray<const int32> WeightView { 0, 1 };
-
-		FStackAttributeContainer BlendedAttributes;
-		UE::Anim::Attributes::BlendAttributes(SourceAttributes, SourceWeights, WeightView, BlendedAttributes);
-		OutAttributes.MoveFrom(BlendedAttributes);
+		Output.CustomAttributes.MoveFrom(OutputAttributes);
 	}
 }
 
diff --git a/Engine/Plugins/Experimental/Animation/HierarchyTableAnimation/Source/Runtime/Private/HierarchyTableBlendProfile.cpp b/Engine/Plugins/Experimental/Animation/HierarchyTableAnimation/Source/Runtime/Private/HierarchyTableBlendProfile.cpp
--- a/Engine/Plugins/Experimental/Animation/HierarchyTableAnimation/Source/Runtime/Private/HierarchyTableBlendProfile.cpp
+++ b/Engine/Plugins/Experimental/Animation/HierarchyTableAnimation/Source/Runtime/Private/HierarchyTableBlendProfile.cpp
@@ -60,8 +60,13 @@ FHierarchyTableBlendProfile::FHierarchyTableBlendProfile(TObjectPtr<UHierarchyTa
 				}
 				else if (Entry.TablePayload.Get<FHierarchyTable_TablePayloadType_Skeleton>().EntryType == ESkeletonHierarchyTable_TablePayloadEntryType::Attribute)
 				{
+					const FHierarchyTableEntryData* ParentEntry = InHierarchyTable->GetTableEntry(Entry.Parent);
+					check(ParentEntry && ParentEntry->GetMetadata<FHierarchyTable_TablePayloadType_Skeleton>().EntryType == ESkeletonHierarchyTable_TablePayloadEntryType::Bone);
+
+					UE::Anim::FAttributeId Attribute(Entry.Identifier, Entry.Parent, TEXT("bone"));
+
 					const float EntryWeight = Entry.GetValue<FHierarchyTable_ElementType_Mask>()->Value;
-					AttributeBlendWeights.Add(FNamedFloat(Entry.Identifier, EntryWeight));
+					AttributeBlendWeights.Add(FMaskedAttributeWeight(Attribute, EntryWeight));
 				}
 			}
 		}
@@ -122,7 +127,7 @@ const UE::Anim::TNamedValueArray<FDefaultAllocator, UE::Anim::FCurveElement>& FH
 	return CurveBlendWeights;
 }
 
-const TArray<FHierarchyTableBlendProfile::FNamedFloat>& FHierarchyTableBlendProfile::GetAttributeBlendWeights() const
+const TArray<FHierarchyTableBlendProfile::FMaskedAttributeWeight>& FHierarchyTableBlendProfile::GetAttributeBlendWeights() const
 {
 	return AttributeBlendWeights;
 }
diff --git a/Engine/Plugins/Experimental/Animation/HierarchyTableAnimation/Source/Runtime/Public/HierarchyTableBlendProfile.h b/Engine/Plugins/Experimental/Animation/HierarchyTableAnimation/Source/Runtime/Public/HierarchyTableBlendProfile.h
--- a/Engine/Plugins/Experimental/Animation/HierarchyTableAnimation/Source/Runtime/Public/HierarchyTableBlendProfile.h
+++ b/Engine/Plugins/Experimental/Animation/HierarchyTableAnimation/Source/Runtime/Public/HierarchyTableBlendProfile.h
@@ -2,6 +2,7 @@
 
 #pragma once
 
+#include "Animation/AttributesContainer.h"
 #include "Animation/BlendProfile.h"
 #include "HierarchyTable.h"
 
@@ -23,20 +24,20 @@ public:
 	virtual EBlendProfileMode GetMode() const override;
 	virtual TObjectPtr<USkeleton> GetSkeleton() const override;
 
-	struct FNamedFloat
+	struct FMaskedAttributeWeight
 	{
-		FName Name;
-		float Value;
+		UE::Anim::FAttributeId Attribute;
+		float Weight;
 
-		FNamedFloat(const FName InName, const float InValue)
-			: Name(InName)
-			, Value(InValue)
+		FMaskedAttributeWeight(const UE::Anim::FAttributeId InAttribute, const float InWeight)
+			: Attribute(InAttribute)
+			, Weight(InWeight)
 		{
 		}
 	};
 
 	const UE::Anim::TNamedValueArray<FDefaultAllocator, UE::Anim::FCurveElement>& GetCurveBlendWeights() const;
-	const TArray<FNamedFloat>& GetAttributeBlendWeights() const;
+	const TArray<FMaskedAttributeWeight>& GetAttributeBlendWeights() const;
 
 	void ConstructBlendProfile(const TObjectPtr<UBlendProfile> OutBlendProfile) const;
 
@@ -49,7 +50,7 @@ private:
 
 	UE::Anim::TNamedValueArray<FDefaultAllocator, UE::Anim::FCurveElement> CurveBlendWeights;
 
-	TArray<FNamedFloat> AttributeBlendWeights;
+	TArray<FMaskedAttributeWeight> AttributeBlendWeights;
 
 	EBlendProfileMode Mode = EBlendProfileMode::WeightFactor;
 };
\ No newline at end of file
