Various changes to Profile Blend node and the Blend Profile asset. [cherrypick]
* Update Profile Blend node to use blend profile asset instead of hierarchy table asset.
* Update Blend Profile asset to synchronise changes with the bound skeleton.
* Update Profile Blend node to use cooked flattened data instead of computing at runtime.

diff --git a/Engine/Plugins/Experimental/Animation/HierarchyTable/Source/Editor/Private/SHierarchyTable.cpp b/Engine/Plugins/Experimental/Animation/HierarchyTable/Source/Editor/Private/SHierarchyTable.cpp
--- a/Engine/Plugins/Experimental/Animation/HierarchyTable/Source/Editor/Private/SHierarchyTable.cpp
+++ b/Engine/Plugins/Experimental/Animation/HierarchyTable/Source/Editor/Private/SHierarchyTable.cpp
@@ -24,11 +24,6 @@ void SHierarchyTable::Construct(const FArguments& InArgs, TObjectPtr<UHierarchyT
 
 	HierarchyTable = InHierarchyTable;
 
-	OnHierarchyChangedHandle = HierarchyTable->RegisterOnHierarchyChanged(UHierarchyTable::FOnHierarchyChanged::FDelegate::CreateLambda([this]()
-	{
-		RegenerateTreeViewItems();
-	}));
-
 	TSharedPtr<SHeaderRow> HeaderRow;
 
 	ChildSlot
@@ -85,10 +80,6 @@ void SHierarchyTable::Construct(const FArguments& InArgs, TObjectPtr<UHierarchyT
 
 SHierarchyTable::~SHierarchyTable()
 {
-	if (HierarchyTable)
-	{
-		HierarchyTable->UnregisterOnHierarchyChanged(OnHierarchyChangedHandle);
-	}
 }
 
 int32 SHierarchyTable::GetSelectedEntryIndex() const
@@ -253,7 +244,7 @@ void SHierarchyTable::RegenerateTreeViewItems()
 	// Update tree view and restore tree item expanded states
 	{
 		check(TreeView);
-		TreeView->RequestTreeRefresh();
+		TreeView->RebuildList();
 
 		for (TSharedPtr<FTreeItem> TreeItem : GetAllTreeItems())
 		{
@@ -287,4 +278,13 @@ FString SHierarchyTable::GetReferencerName() const
 	return FString(TEXT("SHierarchyTable"));
 }
 
+void SHierarchyTable::Tick(const FGeometry& AllottedGeometry, const double InCurrentTime, const float InDeltaTime)
+{
+	if (HierarchyTable->GetHierarchyGuid() != TableHierarchyGuid)
+	{
+		TableHierarchyGuid = HierarchyTable->GetHierarchyGuid();
+		RegenerateTreeViewItems();
+	}
+}
+
 #undef LOCTEXT_NAMESPACE
\ No newline at end of file
diff --git a/Engine/Plugins/Experimental/Animation/HierarchyTable/Source/Editor/Private/SHierarchyTable.h b/Engine/Plugins/Experimental/Animation/HierarchyTable/Source/Editor/Private/SHierarchyTable.h
--- a/Engine/Plugins/Experimental/Animation/HierarchyTable/Source/Editor/Private/SHierarchyTable.h
+++ b/Engine/Plugins/Experimental/Animation/HierarchyTable/Source/Editor/Private/SHierarchyTable.h
@@ -41,6 +41,9 @@ public:
 	void Construct(const FArguments& InArgs, TObjectPtr<UHierarchyTable> InHierarchyTable);
 	~SHierarchyTable();
 
+	// SCompoundWidget
+	virtual void Tick(const FGeometry& AllottedGeometry, const double InCurrentTime, const float InDeltaTime) override;
+
 	// IHierarchyTable
 	int32 GetSelectedEntryIndex() const override;
 
@@ -77,5 +80,5 @@ private:
 
 	TObjectPtr<UHierarchyTable_TableTypeHandler> TableHandler;
 
-	FDelegateHandle OnHierarchyChangedHandle;
+	FGuid TableHierarchyGuid;
 };
diff --git a/Engine/Plugins/Experimental/Animation/HierarchyTable/Source/Runtime/Private/HierarchyTable.cpp b/Engine/Plugins/Experimental/Animation/HierarchyTable/Source/Runtime/Private/HierarchyTable.cpp
--- a/Engine/Plugins/Experimental/Animation/HierarchyTable/Source/Runtime/Private/HierarchyTable.cpp
+++ b/Engine/Plugins/Experimental/Animation/HierarchyTable/Source/Runtime/Private/HierarchyTable.cpp
@@ -44,6 +44,11 @@ void FHierarchyTableEntryData::ToggleOverridden()
 	}
 }
 
+const TOptional<FInstancedStruct>& FHierarchyTableEntryData::GetPayload() const
+{
+	return Payload;
+}
+
 const FInstancedStruct* FHierarchyTableEntryData::GetActualValue() const
 {
 	return IsOverridden() ? Payload.GetPtrOrNull() : GetFromClosestAncestor();
@@ -104,9 +109,8 @@ int32 UHierarchyTable::AddEntry(const FHierarchyTableEntryData& Entry)
 	}
 
 	const int32 EntryIndex = TableData.Add(Entry);
-#if WITH_EDITOR
-	OnHierarchyChanged.Broadcast();
-#endif //WITH_EDITOR
+	RegenerateHierarchyGuid();
+
 	return EntryIndex;
 }
 
@@ -120,9 +124,7 @@ void UHierarchyTable::AddBulkEntries(const TConstArrayView<FHierarchyTableEntryD
 		}
 	}
 
-#if WITH_EDITOR
-	OnHierarchyChanged.Broadcast();
-#endif //WITH_EDITOR
+	RegenerateHierarchyGuid();
 }
 
 void UHierarchyTable::RemoveEntry(const int32 IndexToRemove)
@@ -148,23 +150,8 @@ void UHierarchyTable::RemoveEntry(const int32 IndexToRemove)
 	}
 
 	TableData.RemoveAt(IndexToRemove);
-
-#if WITH_EDITOR
-	OnHierarchyChanged.Broadcast();
-#endif //WITH_EDITOR
-}
-
-#if WITH_EDITOR
-FDelegateHandle UHierarchyTable::RegisterOnHierarchyChanged(const FOnHierarchyChanged::FDelegate& Delegate)
-{
-	return OnHierarchyChanged.Add(Delegate);
-}
-
-bool UHierarchyTable::UnregisterOnHierarchyChanged(const FDelegateHandle& Handle)
-{
-	return OnHierarchyChanged.Remove(Handle);
+	RegenerateHierarchyGuid();
 }
-#endif //WITH_EDITOR
 
 FInstancedStruct UHierarchyTable::CreateDefaultValue() const
 {
@@ -174,6 +161,19 @@ FInstancedStruct UHierarchyTable::CreateDefaultValue() const
 	return OutStruct;
 }
 
+const FHierarchyTableEntryData* const UHierarchyTable::GetTableEntry(const FName EntryIdentifier) const
+{
+	for (const FHierarchyTableEntryData& Entry : TableData)
+	{
+		if (Entry.Identifier == EntryIdentifier)
+		{
+			return &Entry;
+		}
+	}
+
+	return nullptr;
+}
+
 const FHierarchyTableEntryData* const UHierarchyTable::GetTableEntry(const int32 EntryIndex) const
 {
 	if (TableData.IsValidIndex(EntryIndex))
@@ -195,8 +195,28 @@ FHierarchyTableEntryData* const UHierarchyTable::GetMutableTableEntry(const int3
 void UHierarchyTable::EmptyTable()
 {
 	TableData.Empty();
+	RegenerateHierarchyGuid();
+}
+
+FGuid UHierarchyTable::GetHierarchyGuid()
+{
+	return HierarchyGuid;
+}
+
+void UHierarchyTable::RegenerateHierarchyGuid()
+{
+	HierarchyGuid = FGuid::NewGuid();
+}
+
+int32 UHierarchyTable::GetTableEntryIndex(const FName EntryIdentifier) const
+{
+	for (int32 EntryIndex = 0; EntryIndex < TableData.Num(); ++EntryIndex)
+	{
+		if (TableData[EntryIndex].Identifier == EntryIdentifier)
+		{
+			return EntryIndex;
+		}
+	}
 
-#if WITH_EDITOR
-	OnHierarchyChanged.Broadcast();
-#endif //WITH_EDITOR
+	return INDEX_NONE;
 }
\ No newline at end of file
diff --git a/Engine/Plugins/Experimental/Animation/HierarchyTable/Source/Runtime/Public/HierarchyTable.h b/Engine/Plugins/Experimental/Animation/HierarchyTable/Source/Runtime/Public/HierarchyTable.h
--- a/Engine/Plugins/Experimental/Animation/HierarchyTable/Source/Runtime/Public/HierarchyTable.h
+++ b/Engine/Plugins/Experimental/Animation/HierarchyTable/Source/Runtime/Public/HierarchyTable.h
@@ -58,6 +58,9 @@ public:
 
 	void ToggleOverridden();
 
+	// Assumes payload is set, i.e. is overridden
+	const TOptional<FInstancedStruct>& GetPayload() const;
+
 	template <typename T>
 	const T* GetValue() const
 	{
@@ -158,8 +161,15 @@ public:
 
 	void EmptyTable();
 
+	// TODO: Remove in the future to avoid API signatures using indices
+	int32 GetTableEntryIndex(const FName EntryIdentifier) const;
+
+	const FHierarchyTableEntryData* const GetTableEntry(const FName EntryIdentifier) const;
+
+	// TODO: Remove in the future to avoid API signatures using indices
 	const FHierarchyTableEntryData* const GetTableEntry(const int32 EntryIndex) const;
 
+	// TODO: Remove in the future to avoid API signatures using indices
 	FHierarchyTableEntryData* const GetMutableTableEntry(const int32 EntryIndex);
 
 	int32 AddEntry(const FHierarchyTableEntryData& Entry);
@@ -181,13 +191,11 @@ public:
 		});
 	}
 
-#if WITH_EDITOR
-	DECLARE_MULTICAST_DELEGATE(FOnHierarchyChanged);
+	FGuid GetHierarchyGuid();
+
 private:
-	FOnHierarchyChanged OnHierarchyChanged;
+	void RegenerateHierarchyGuid();
 
-public:
-	FDelegateHandle RegisterOnHierarchyChanged(const FOnHierarchyChanged::FDelegate& Delegate);
-	bool UnregisterOnHierarchyChanged(const FDelegateHandle& Handle);
-#endif
+	UPROPERTY()
+	FGuid HierarchyGuid;
 };
diff --git a/Engine/Plugins/Experimental/Animation/HierarchyTableAnimation/Source/Editor/Private/BlendProfileStandaloneEditorToolkit.cpp b/Engine/Plugins/Experimental/Animation/HierarchyTableAnimation/Source/Editor/Private/BlendProfileStandaloneEditorToolkit.cpp
--- a/Engine/Plugins/Experimental/Animation/HierarchyTableAnimation/Source/Editor/Private/BlendProfileStandaloneEditorToolkit.cpp
+++ b/Engine/Plugins/Experimental/Animation/HierarchyTableAnimation/Source/Editor/Private/BlendProfileStandaloneEditorToolkit.cpp
@@ -12,6 +12,7 @@
 void FBlendProfileStandaloneEditorToolkit::InitEditor(const TArray<UObject*>& InObjects)
 {
 	BlendProfile = CastChecked<UBlendProfileStandalone>(InObjects[0]);
+	BlendProfile->UpdateHierarchy();
 
 	const TSharedRef<FTabManager::FLayout> Layout = FTabManager::NewLayout("BlendProfileStandaloneEditorToolkit")
 		->AddArea
@@ -34,6 +35,21 @@ void FBlendProfileStandaloneEditorToolkit::InitEditor(const TArray<UObject*>& In
 	FAssetEditorToolkit::InitAssetEditor(EToolkitMode::Standalone, {}, "BlendProfileStandaloneEditor", Layout, true, true, InObjects);
 
 	ExtendToolbar();
+
+	TObjectPtr<USkeleton> Skeleton = BlendProfile->GetSkeleton();
+	if (Skeleton)
+	{
+		Skeleton->RegisterOnSkeletonHierarchyChanged(USkeleton::FOnSkeletonHierarchyChanged::CreateRaw(this, &FBlendProfileStandaloneEditorToolkit::OnSkeletonHierarchyChanged));
+	}
+}
+
+void FBlendProfileStandaloneEditorToolkit::OnClose()
+{
+	TObjectPtr<USkeleton> Skeleton = BlendProfile->GetSkeleton();
+	if (Skeleton)
+	{
+		Skeleton->UnregisterOnSkeletonHierarchyChanged(this);
+	}
 }
 
 void FBlendProfileStandaloneEditorToolkit::RegisterTabSpawners(const TSharedRef<class FTabManager>& InTabManager)
@@ -88,6 +104,12 @@ void FBlendProfileStandaloneEditorToolkit::UnregisterTabSpawners(const TSharedRe
 	InTabManager->UnregisterTabSpawner("BlendProfileStandaloneEditorDetailsTab");
 }
 
+void FBlendProfileStandaloneEditorToolkit::OnSkeletonHierarchyChanged()
+{
+	BlendProfile->UpdateHierarchy();
+	BlendProfile->UpdateCachedData();
+}
+
 void FBlendProfileStandaloneEditorToolkit::ExtendToolbar()
 {
 	FHierarchyTableEditorModule& HierarchyTableModule = FModuleManager::GetModuleChecked<FHierarchyTableEditorModule>("HierarchyTableEditor");
diff --git a/Engine/Plugins/Experimental/Animation/HierarchyTableAnimation/Source/Editor/Private/BlendProfileStandaloneEditorToolkit.h b/Engine/Plugins/Experimental/Animation/HierarchyTableAnimation/Source/Editor/Private/BlendProfileStandaloneEditorToolkit.h
--- a/Engine/Plugins/Experimental/Animation/HierarchyTableAnimation/Source/Editor/Private/BlendProfileStandaloneEditorToolkit.h
+++ b/Engine/Plugins/Experimental/Animation/HierarchyTableAnimation/Source/Editor/Private/BlendProfileStandaloneEditorToolkit.h
@@ -15,6 +15,8 @@ public:
 	void InitEditor(const TArray<UObject*>& InObjects);
 
 	// Begin FAssetEditorToolkit
+	void OnClose() override;
+
 	void RegisterTabSpawners(const TSharedRef<class FTabManager>& TabManager) override;
 	void UnregisterTabSpawners(const TSharedRef<class FTabManager>& TabManager) override;
 
@@ -25,6 +27,8 @@ public:
 	// End FAssetEditorToolkit
 
 private:
+	void OnSkeletonHierarchyChanged();
+
 	void ExtendToolbar();
 
 	TObjectPtr<UBlendProfileStandalone> BlendProfile;
diff --git a/Engine/Plugins/Experimental/Animation/HierarchyTableAnimation/Source/Runtime/Private/AnimNode_BlendProfileLayeredBlend.cpp b/Engine/Plugins/Experimental/Animation/HierarchyTableAnimation/Source/Runtime/Private/AnimNode_BlendProfileLayeredBlend.cpp
--- a/Engine/Plugins/Experimental/Animation/HierarchyTableAnimation/Source/Runtime/Private/AnimNode_BlendProfileLayeredBlend.cpp
+++ b/Engine/Plugins/Experimental/Animation/HierarchyTableAnimation/Source/Runtime/Private/AnimNode_BlendProfileLayeredBlend.cpp
@@ -12,6 +12,7 @@
 #include "Animation/IAttributeBlendOperator.h"
 #include "SkeletonHierarchyTableType.h"
 #include "HierarchyTableBlendProfile.h"
+#include "BlendProfileStandalone.h"
 
 #include UE_INLINE_GENERATED_CPP_BY_NAME(AnimNode_BlendProfileLayeredBlend)
 
@@ -26,45 +27,28 @@ void FAnimNode_BlendProfileLayeredBlend::Initialize_AnyThread(const FAnimationIn
 
 bool FAnimNode_BlendProfileLayeredBlend::ArePerBoneBlendWeightsValid(const USkeleton* InSkeleton) const
 {
-	return (InSkeleton != nullptr && InSkeleton->GetGuid() == SkeletonGuid && InSkeleton->GetVirtualBoneGuid() == VirtualBoneGuid && CachedMaskTable == MaskTable);
+	return (InSkeleton != nullptr && InSkeleton->GetGuid() == SkeletonGuid && InSkeleton->GetVirtualBoneGuid() == VirtualBoneGuid && CachedBlendProfile == BlendProfileAsset);
 }
 
 void FAnimNode_BlendProfileLayeredBlend::UpdateCachedBoneData(const FBoneContainer& RequiredBones, USkeleton* Skeleton)
 {
-
 	if (ArePerBoneBlendWeightsValid(Skeleton) && RequiredBonesSerialNumber == RequiredBones.GetSerialNumber())
 	{
 		return;
 	}
 
 	QUICK_SCOPE_CYCLE_COUNTER(STAT_BlendProfileLayeredBlend_UpdateCachedBoneData);
-	if (Skeleton && MaskTable && MaskTable->IsTableType<FHierarchyTable_TableType_Skeleton>())
-	{
-		const FHierarchyTable_TableType_Skeleton& SkeletonMetadata = MaskTable->GetTableMetadata<FHierarchyTable_TableType_Skeleton>();
-
-		if (SkeletonMetadata.Skeleton == Skeleton)
-		{
-			HierarchyTableBlendProfile = FHierarchyTableBlendProfile(MaskTable, EBlendProfileMode::BlendMask);
-		}
-	}
-	else
-	{
-		HierarchyTableBlendProfile = FHierarchyTableBlendProfile(Skeleton, EBlendProfileMode::BlendMask);
-	}
 	
+	const TArray<FBoneIndexType>& RequiredBoneIndices = RequiredBones.GetBoneIndicesArray();
+	const int32 NumRequiredBones = RequiredBoneIndices.Num();
+
 	// Update DesiredBoneBlendWeights
+	const bool bBlendProfileValid = Skeleton && BlendProfileAsset && BlendProfileAsset->Type == EBlendProfileStandaloneType::BlendMask && BlendProfileAsset->GetSkeleton() == Skeleton;
+	if (bBlendProfileValid)
 	{
-		const TArray<FBoneIndexType>& RequiredBoneIndices = RequiredBones.GetBoneIndicesArray();
-		const int32 NumRequiredBones = RequiredBoneIndices.Num();
-		DesiredBoneBlendWeights.SetNumZeroed(NumRequiredBones);
-
-		for (int32 RequiredBoneIndex = 0; RequiredBoneIndex < NumRequiredBones; RequiredBoneIndex++)
+		if (ensure(NumRequiredBones <= BlendProfileAsset->CachedBlendProfileData.GetBoneBlendWeights().Num()))
 		{
-			const int32 SkeletonBoneIndex = RequiredBones.GetSkeletonIndex(FCompactPoseBoneIndex(RequiredBoneIndex));
-			if (HierarchyTableBlendProfile.IsValidBoneIndex(SkeletonBoneIndex))
-			{
-				DesiredBoneBlendWeights[RequiredBoneIndex] = HierarchyTableBlendProfile.GetBoneBlendScale(SkeletonBoneIndex);
-			}
+			DesiredBoneBlendWeights = BlendProfileAsset->CachedBlendProfileData.GetBoneBlendWeights();
 		}
 	}
 
@@ -72,15 +56,15 @@ void FAnimNode_BlendProfileLayeredBlend::UpdateCachedBoneData(const FBoneContain
 	{
 		CurrentBoneBlendWeights.Reset(DesiredBoneBlendWeights.Num());
 		CurrentBoneBlendWeights.AddZeroed(DesiredBoneBlendWeights.Num());
-
+		
 		//Reinitialize bone blend weights now that we have cleared them
 		UpdateDesiredBoneWeight();
 	}
 
 	SkeletonGuid = Skeleton ? Skeleton->GetGuid() : FGuid();
 	VirtualBoneGuid = Skeleton ? Skeleton->GetVirtualBoneGuid() : FGuid();
-	CachedMaskTable = MaskTable;
 	RequiredBonesSerialNumber = RequiredBones.GetSerialNumber();
+	CachedBlendProfile = BlendProfileAsset;
 }
 
 void FAnimNode_BlendProfileLayeredBlend::CacheBones_AnyThread(const FAnimationCacheBonesContext& Context)
@@ -143,7 +127,10 @@ void FAnimNode_BlendProfileLayeredBlend::Evaluate_AnyThread(FPoseContext& Output
 	DECLARE_SCOPE_HIERARCHICAL_COUNTER_ANIMNODE(Evaluate_AnyThread)
 	ANIM_MT_SCOPE_CYCLE_COUNTER(BlendPosesInGraph, !IsInGameThread());
 
-	if (!FAnimWeight::IsRelevant(BlendWeight))
+	TObjectPtr<USkeleton> Skeleton = Output.AnimInstanceProxy->GetSkeleton();
+	const bool bBlendProfileValid = Skeleton && BlendProfileAsset && BlendProfileAsset->Type == EBlendProfileStandaloneType::BlendMask && BlendProfileAsset->GetSkeleton() == Skeleton;
+
+	if (!FAnimWeight::IsRelevant(BlendWeight) || !bBlendProfileValid)
 	{
 		BasePose.Evaluate(Output);
 		return;
@@ -174,10 +161,11 @@ void FAnimNode_BlendProfileLayeredBlend::Evaluate_AnyThread(FPoseContext& Output
 		}
 	}
 
+
 	// Blend curves
 	if (!bCustomCurveBlending)
 	{
-		const UE::Anim::TNamedValueArray<FDefaultAllocator, UE::Anim::FCurveElement>& CachedCurveMaskWeights = HierarchyTableBlendProfile.GetCurveBlendWeights();
+		const UE::Anim::TNamedValueArray<FDefaultAllocator, UE::Anim::FCurveElement>& CachedCurveMaskWeights = BlendProfileAsset->CachedBlendProfileData.GetCurveBlendWeights();
 
 		Output.Curve.CopyFrom(BasePoseData.GetCurve());
 
@@ -240,6 +228,8 @@ void FAnimNode_BlendProfileLayeredBlend::Evaluate_AnyThread(FPoseContext& Output
 
 	// Blend attributes
 	{
+		const TArray<FBlendProfileStandaloneCachedData::FMaskedAttributeWeight>& CachedAttributeMaskWeights = BlendProfileAsset->CachedBlendProfileData.GetAttributeBlendWeights();
+		
 		using namespace UE::Anim;
 
 		FStackAttributeContainer OutputAttributes;
@@ -279,8 +269,7 @@ void FAnimNode_BlendProfileLayeredBlend::Evaluate_AnyThread(FPoseContext& Output
 		UE::Anim::Attributes::BlendAttributesPerBone(BasePoseData.GetAttributes(), BlendPoseData.GetAttributes(), CurrentBoneBlendWeights, OutputAttributes);
 
 		// 2. For each attribute that has a custom weight, i.e. one's that shouldn't be weighted by its attached bone, go and correct the blended value.
-		const TArray<FHierarchyTableBlendProfile::FMaskedAttributeWeight>& CachedAttributeMaskWeights = HierarchyTableBlendProfile.GetAttributeBlendWeights();
-		for (const FHierarchyTableBlendProfile::FMaskedAttributeWeight& MaskedAttribute : CachedAttributeMaskWeights)
+		for (const FBlendProfileStandaloneCachedData::FMaskedAttributeWeight& MaskedAttribute : CachedAttributeMaskWeights)
 		{
 			const TArray<TWeakObjectPtr<UScriptStruct>>& UniqueTypes = OutputAttributes.GetUniqueTypes();
 
@@ -351,8 +340,11 @@ void FAnimNode_BlendProfileLayeredBlend::GatherDebugData(FNodeDebugData& DebugDa
 {
 	DECLARE_SCOPE_HIERARCHICAL_COUNTER_ANIMNODE(GatherDebugData)
 
+	TObjectPtr<USkeleton> Skeleton = DebugData.AnimInstance->CurrentSkeleton;
+	const bool bBlendProfileValid = Skeleton && BlendProfileAsset && BlendProfileAsset->Type == EBlendProfileStandaloneType::BlendMask && BlendProfileAsset->GetSkeleton() == Skeleton;
+
 	BasePose.GatherDebugData(DebugData.BranchFlow(1.f));
-	BlendPose.GatherDebugData(DebugData.BranchFlow(BlendWeight));
+	BlendPose.GatherDebugData(DebugData.BranchFlow(bBlendProfileValid ? BlendWeight : 0.0f));
 }
 
 void FAnimNode_BlendProfileLayeredBlend::UpdateDesiredBoneWeight()
diff --git a/Engine/Plugins/Experimental/Animation/HierarchyTableAnimation/Source/Runtime/Private/BlendProfileStandalone.cpp b/Engine/Plugins/Experimental/Animation/HierarchyTableAnimation/Source/Runtime/Private/BlendProfileStandalone.cpp
--- a/Engine/Plugins/Experimental/Animation/HierarchyTableAnimation/Source/Runtime/Private/BlendProfileStandalone.cpp
+++ b/Engine/Plugins/Experimental/Animation/HierarchyTableAnimation/Source/Runtime/Private/BlendProfileStandalone.cpp
@@ -7,10 +7,196 @@
 #include "SkeletonHierarchyTableType.h"
 #include "UObject/AssetRegistryTagsContext.h"
 #include "UObject/NoExportTypes.h"
+#include "UObject/ObjectSaveContext.h"
+#include "StructUtils/InstancedStruct.h"
+#include "Misc/Optional.h"
 #endif // WITH_EDITOR
 
 #if WITH_EDITOR
 
+void UBlendProfileStandalone::UpdateHierarchy()
+{
+	/* Skeleton hierarchy tables take a snapshot of the currently bound skeleton's hierarchy.
+	 *  When the skeleton hierarchy changes and this function is called we do the following:
+	 *		1. Make a note of all the bone table entries with overridden values.
+	 *		2. Make a note of all the curves and attributes in the table and their location in the hierarchy.
+	 *		3. Clear the hierarchy table contents
+	 *		4. Reconstruct the bone hierarchy using the new skeleton hierarchy and insert the old blend values 
+	 *		   for the old overridden bones.
+	 *		5. Reconstruct the curves and attributes hierarchy.
+	 */
+
+	const FHierarchyTable_TableType_Skeleton& Metadata = Table->GetTableMetadata<FHierarchyTable_TableType_Skeleton>();
+	Skeleton = Metadata.Skeleton;
+
+	if (Skeleton->GetGuid() == SkeletonHierarchyGuid && Skeleton->GetVirtualBoneGuid() == SkeletonVirtualBonesHierarchyGuid)
+	{
+		return;
+	}
+	SkeletonHierarchyGuid = Skeleton->GetGuid();
+	SkeletonVirtualBonesHierarchyGuid =  Skeleton->GetVirtualBoneGuid();
+
+	const FReferenceSkeleton& RefSkeleton = GetSkeleton()->GetReferenceSkeleton();
+	const int32 BoneCount = RefSkeleton.GetNum();
+
+	struct FCurveAttributeBlendData
+	{
+		bool bIsCurve;
+		FName Identifier;
+		FName ParentIdentifier;
+		TOptional<FInstancedStruct> Payload;
+	};
+
+	TMap<FName, FInstancedStruct> BoneData;
+	TArray<FCurveAttributeBlendData> CurveAndAttributeData;
+	{
+		for (const FHierarchyTableEntryData& TableEntry : Table->GetTableData())
+		{
+			const auto& EntryMetadata = TableEntry.GetMetadata<FHierarchyTable_TablePayloadType_Skeleton>();
+			if (EntryMetadata.EntryType == ESkeletonHierarchyTable_TablePayloadEntryType::Bone)
+			{
+				if (TableEntry.IsOverridden())
+				{
+					BoneData.Add(TableEntry.Identifier, TableEntry.GetPayload().GetValue());
+				}
+			}
+			if (EntryMetadata.EntryType == ESkeletonHierarchyTable_TablePayloadEntryType::Curve)
+			{
+				check(TableEntry.HasParent()); // Curves must be parented
+				const int32 ParentIndex = TableEntry.Parent;
+
+				const FHierarchyTableEntryData* ParentEntry = Table->GetTableEntry(ParentIndex);
+				check(ParentEntry);
+
+				CurveAndAttributeData.Add({ true, TableEntry.Identifier, ParentEntry->Identifier, TableEntry.GetPayload() });
+			}
+			else if (EntryMetadata.EntryType == ESkeletonHierarchyTable_TablePayloadEntryType::Attribute)
+			{
+				check(TableEntry.HasParent()); // Attributes must be parented
+				const int32 ParentIndex = TableEntry.Parent;
+
+				const FHierarchyTableEntryData* ParentEntry = Table->GetTableEntry(ParentIndex);
+				check(ParentEntry);
+
+				CurveAndAttributeData.Add({ false, TableEntry.Identifier, ParentEntry->Identifier, TableEntry.GetPayload() });
+			}
+		}
+	}
+
+	Table->EmptyTable();
+
+	// Re-add bone data
+	{
+		FInstancedStruct DefaultTablePayload;
+		DefaultTablePayload.InitializeAs<FHierarchyTable_TablePayloadType_Skeleton>();
+		DefaultTablePayload.GetMutable<FHierarchyTable_TablePayloadType_Skeleton>().EntryType = ESkeletonHierarchyTable_TablePayloadEntryType::Bone;
+
+		for (int32 BoneIndex = 0; BoneIndex < BoneCount; ++BoneIndex)
+		{
+			const FName BoneName = RefSkeleton.GetBoneName(BoneIndex);
+
+			FHierarchyTableEntryData EntryData;
+			EntryData.Parent = RefSkeleton.GetParentIndex(BoneIndex);
+			EntryData.Identifier = BoneName;
+			EntryData.TablePayload = DefaultTablePayload;
+			EntryData.OwnerTable = Table;
+
+			if (BoneData.Contains(BoneName))
+			{
+				EntryData.Payload = BoneData[BoneName];
+			}
+			else
+			{
+				EntryData.Payload = (BoneIndex == 0) ? Table->CreateDefaultValue() : TOptional<FInstancedStruct>();
+			}
+
+			Table->AddEntry(EntryData);
+		}
+	}
+
+	// Re-add curve and attribute data
+	{
+		FInstancedStruct DefaultTablePayload_Curve;
+		DefaultTablePayload_Curve.InitializeAs<FHierarchyTable_TablePayloadType_Skeleton>();
+		DefaultTablePayload_Curve.GetMutable<FHierarchyTable_TablePayloadType_Skeleton>().EntryType = ESkeletonHierarchyTable_TablePayloadEntryType::Curve;
+
+		FInstancedStruct DefaultTablePayload_Attribute;
+		DefaultTablePayload_Attribute.InitializeAs<FHierarchyTable_TablePayloadType_Skeleton>();
+		DefaultTablePayload_Attribute.GetMutable<FHierarchyTable_TablePayloadType_Skeleton>().EntryType = ESkeletonHierarchyTable_TablePayloadEntryType::Attribute;
+
+		for (const FCurveAttributeBlendData& Entry : CurveAndAttributeData)
+		{
+			const FHierarchyTableEntryData* const ParentEntry = Table->GetTableEntry(Entry.ParentIdentifier);
+			if (ParentEntry)
+			{
+				FHierarchyTableEntryData EntryData;
+				EntryData.Identifier = Entry.Identifier;
+				EntryData.Parent = Table->GetTableEntryIndex(ParentEntry->Identifier);
+				EntryData.TablePayload = Entry.bIsCurve ? DefaultTablePayload_Curve : DefaultTablePayload_Attribute;
+				EntryData.OwnerTable = Table;
+				EntryData.Payload = Entry.Payload;
+
+				Table->AddEntry(EntryData);
+			}
+			else
+			{
+				// This curve/attribute was attached to a parent that no longer exists
+			}
+		}
+	}
+}
+
+void UBlendProfileStandalone::UpdateCachedData()
+{
+	CachedBlendProfileData.Reset();
+
+	if (Table)
+	{
+		EBlendProfileMode Mode = EBlendProfileMode::WeightFactor;
+		if (Type == EBlendProfileStandaloneType::TimeFactor)
+		{
+			Mode = EBlendProfileMode::TimeFactor;
+		}
+		else if (Type == EBlendProfileStandaloneType::BlendMask)
+		{
+			Mode = EBlendProfileMode::BlendMask;
+		}
+
+		CachedBlendProfileData.Init(Table, Mode);
+
+		const FHierarchyTable_TableType_Skeleton& Metadata = Table->GetTableMetadata<FHierarchyTable_TableType_Skeleton>();
+		Skeleton = Metadata.Skeleton;
+	}
+}
+
+#endif // WITH_EDITOR
+
+TObjectPtr<USkeleton> UBlendProfileStandalone::GetSkeleton() const
+{
+	return Skeleton;
+}
+
+void UBlendProfileStandalone::PostLoad()
+{
+	Super::PostLoad();
+
+#if WITH_EDITOR
+	// Update the hierarchy to reflect any changes made to the skeleton
+	UpdateHierarchy();
+	// Cache the flattened data for runtime use
+	UpdateCachedData();
+#endif // WITH_EDITOR
+
+	CachedBlendProfileData.UnpackCachedData();
+}
+
+#if WITH_EDITOR
+
+void UBlendProfileStandalone::PreSave(FObjectPreSaveContext ObjectSaveContext)
+{
+	UpdateCachedData();
+}
+
 void UBlendProfileStandalone::GetAssetRegistryTags(FAssetRegistryTagsContext Context) const
 {
 	Super::GetAssetRegistryTags(Context);
diff --git a/Engine/Plugins/Experimental/Animation/HierarchyTableAnimation/Source/Runtime/Private/HierarchyTableBlendProfile.cpp b/Engine/Plugins/Experimental/Animation/HierarchyTableAnimation/Source/Runtime/Private/HierarchyTableBlendProfile.cpp
--- a/Engine/Plugins/Experimental/Animation/HierarchyTableAnimation/Source/Runtime/Private/HierarchyTableBlendProfile.cpp
+++ b/Engine/Plugins/Experimental/Animation/HierarchyTableAnimation/Source/Runtime/Private/HierarchyTableBlendProfile.cpp
@@ -6,6 +6,146 @@
 #include "HierarchyTableDefaultTypes.h"
 #include "MaskProfile/HierarchyTableTypeMask.h"
 
+void FBlendProfileStandaloneCachedData::UnpackCachedData()
+{
+	CurveBlendWeights.Reserve(SerializedCurveBlendWeights.Num());
+	for (const FMaskedCurveWeightSerialised SerializedCurve : SerializedCurveBlendWeights)
+	{
+		CurveBlendWeights.Add(SerializedCurve.CurveName, SerializedCurve.Weight);
+	}
+
+	AttributeBlendWeights.Reserve(SerializedAttributeBlendWeights.Num());
+	for (const FMaskedAttributeWeightSerialised& SerializedAttribute : SerializedAttributeBlendWeights)
+	{
+		UE::Anim::FAttributeId AttributeId(SerializedAttribute.AttributeName, SerializedAttribute.AttributeIndex, SerializedAttribute.AttributeNamespace);
+		AttributeBlendWeights.Add(FMaskedAttributeWeight(AttributeId, SerializedAttribute.Weight));
+	}
+}
+
+void FBlendProfileStandaloneCachedData::Init(TObjectPtr<UHierarchyTable> InHierarchyTable, const EBlendProfileMode InMode)
+{
+	Mode = InMode;
+
+	const FHierarchyTable_TableType_Skeleton& TableMetadata = InHierarchyTable->GetTableMetadata<FHierarchyTable_TableType_Skeleton>();
+	Skeleton = TableMetadata.Skeleton;
+
+	const FReferenceSkeleton RefSkeleton = Skeleton->GetReferenceSkeleton();
+
+	const int32 BoneCount = RefSkeleton.GetNum();
+	BoneBlendWeights.Reserve(BoneCount);
+
+	if (InHierarchyTable->IsElementType<FHierarchyTable_ElementType_Float>())
+	{
+		for (int32 BoneIndex = 0; BoneIndex < BoneCount; ++BoneIndex)
+		{
+			if (const FHierarchyTableEntryData* EntryData = InHierarchyTable->GetTableEntry(BoneIndex))
+			{
+				BoneBlendWeights.Add(EntryData->GetValue<FHierarchyTable_ElementType_Float>()->Value);
+			}
+			else
+			{
+				BoneBlendWeights.Add(0.0f);
+			}
+		}
+	}
+	else if (InHierarchyTable->IsElementType<FHierarchyTable_ElementType_Mask>())
+	{
+		// Update bone weights
+		for (int32 BoneIndex = 0; BoneIndex < BoneCount; ++BoneIndex)
+		{
+			if (const FHierarchyTableEntryData* EntryData = InHierarchyTable->GetTableEntry(BoneIndex))
+			{
+				BoneBlendWeights.Add(EntryData->GetValue<FHierarchyTable_ElementType_Mask>()->Value);
+			}
+			else
+			{
+				BoneBlendWeights.Add(0.0f);
+			}
+		}
+
+		// Update curve and attribute weights
+		{
+			CurveBlendWeights.Empty();
+			AttributeBlendWeights.Empty();
+
+			for (const FHierarchyTableEntryData& Entry : InHierarchyTable->GetTableData())
+			{
+				if (Entry.TablePayload.Get<FHierarchyTable_TablePayloadType_Skeleton>().EntryType == ESkeletonHierarchyTable_TablePayloadEntryType::Curve)
+				{
+					const float EntryWeight = Entry.GetValue<FHierarchyTable_ElementType_Mask>()->Value;
+					CurveBlendWeights.Add<UE::Anim::FCurveElement>({ Entry.Identifier, EntryWeight });
+				}
+				else if (Entry.TablePayload.Get<FHierarchyTable_TablePayloadType_Skeleton>().EntryType == ESkeletonHierarchyTable_TablePayloadEntryType::Attribute)
+				{
+					const FHierarchyTableEntryData* ParentEntry = InHierarchyTable->GetTableEntry(Entry.Parent);
+					check(ParentEntry && ParentEntry->GetMetadata<FHierarchyTable_TablePayloadType_Skeleton>().EntryType == ESkeletonHierarchyTable_TablePayloadEntryType::Bone);
+
+					UE::Anim::FAttributeId Attribute(Entry.Identifier, Entry.Parent, TEXT("bone"));
+
+					const float EntryWeight = Entry.GetValue<FHierarchyTable_ElementType_Mask>()->Value;
+					AttributeBlendWeights.Add(FMaskedAttributeWeight(Attribute, EntryWeight));
+				}
+			}
+		}
+	}
+	else
+	{
+		ensureMsgf(false, TEXT("Unsupported hierarchy table type, use Float or Mask element type instead"));
+		for (int32 BoneIndex = 0; BoneIndex < BoneCount; ++BoneIndex)
+		{
+			BoneBlendWeights.Add(1.0f);
+		}
+	}
+}
+
+TObjectPtr<USkeleton> FBlendProfileStandaloneCachedData::GetSkeleton() const
+{
+	return Skeleton;
+}
+
+const TArray<float>& FBlendProfileStandaloneCachedData::GetBoneBlendWeights() const
+{
+	return BoneBlendWeights;
+}
+
+const UE::Anim::TNamedValueArray<FDefaultAllocator, UE::Anim::FCurveElement>& FBlendProfileStandaloneCachedData::GetCurveBlendWeights() const
+{
+	return CurveBlendWeights;
+}
+
+const TArray<FBlendProfileStandaloneCachedData::FMaskedAttributeWeight>& FBlendProfileStandaloneCachedData::GetAttributeBlendWeights() const
+{
+	return AttributeBlendWeights;
+}
+
+void FBlendProfileStandaloneCachedData::ConstructBlendProfile(const TObjectPtr<UBlendProfile> OutBlendProfile) const
+{
+	OutBlendProfile->SetSkeleton(Skeleton);
+	OutBlendProfile->Mode = Mode;
+
+	for (int32 BoneIndex = 0; BoneIndex < BoneBlendWeights.Num(); ++BoneIndex)
+	{
+		const float BlendWeight = BoneBlendWeights[BoneIndex];
+
+		OutBlendProfile->SetBoneBlendScale(BoneIndex, BlendWeight, false /* bRecurse */, true /* bCreate */);
+	}
+}
+
+void FBlendProfileStandaloneCachedData::Reset()
+{
+	Skeleton = nullptr;
+	BoneBlendWeights.Empty();
+	CurveBlendWeights.Empty();
+	AttributeBlendWeights.Empty();
+}
+
+bool FBlendProfileStandaloneCachedData::IsValidBoneIndex(const int32 BoneIndex) const
+{
+	return BoneBlendWeights.IsValidIndex(BoneIndex);
+}
+
+
+
 FHierarchyTableBlendProfile::FHierarchyTableBlendProfile(TObjectPtr<UHierarchyTable> InHierarchyTable, const EBlendProfileMode InMode)
 	: Mode(InMode)
 {
diff --git a/Engine/Plugins/Experimental/Animation/HierarchyTableAnimation/Source/Runtime/Public/AnimNode_BlendProfileLayeredBlend.h b/Engine/Plugins/Experimental/Animation/HierarchyTableAnimation/Source/Runtime/Public/AnimNode_BlendProfileLayeredBlend.h
--- a/Engine/Plugins/Experimental/Animation/HierarchyTableAnimation/Source/Runtime/Public/AnimNode_BlendProfileLayeredBlend.h
+++ b/Engine/Plugins/Experimental/Animation/HierarchyTableAnimation/Source/Runtime/Public/AnimNode_BlendProfileLayeredBlend.h
@@ -10,8 +10,7 @@
 #include "AnimNode_BlendProfileLayeredBlend.generated.h"
 
 
-class UHierarchyTable;
-enum class EHierarchyTableEntryType : uint8;
+class UBlendProfileStandalone;
 
 USTRUCT(BlueprintInternalUseOnly)
 struct HIERARCHYTABLEANIMATIONRUNTIME_API FAnimNode_BlendProfileLayeredBlend : public FAnimNode_Base
@@ -28,10 +27,10 @@ public:
 	FPoseLink BlendPose;
 
 	/**
-	 * The mask hierarchy table to use to control layering of the pose, curves, and attributes
+	 * The blend profile mask asset to use to control layering of the pose, curves, and attributes
 	 */
 	UPROPERTY(EditAnywhere, BlueprintReadWrite, Category = Config, meta = (PinHiddenByDefault))
-	TObjectPtr<UHierarchyTable> MaskTable;
+	TObjectPtr<UBlendProfileStandalone> BlendProfileAsset;
 
 	/** Whether to blend bone rotations in mesh space or in local space */
 	UPROPERTY(EditAnywhere, BlueprintReadWrite, Category = Config, meta = (PinHiddenByDefault))
@@ -95,11 +94,8 @@ private:
 
 	void UpdateDesiredBoneWeight();
 
-	FHierarchyTableBlendProfile HierarchyTableBlendProfile;
-
 	TArray<float> DesiredBoneBlendWeights;
 	TArray<float> CurrentBoneBlendWeights;
 
-	UPROPERTY()
-	TObjectPtr<UHierarchyTable> CachedMaskTable;
+	TObjectPtr<UBlendProfileStandalone> CachedBlendProfile;
 };
\ No newline at end of file
diff --git a/Engine/Plugins/Experimental/Animation/HierarchyTableAnimation/Source/Runtime/Public/BlendProfileStandalone.h b/Engine/Plugins/Experimental/Animation/HierarchyTableAnimation/Source/Runtime/Public/BlendProfileStandalone.h
--- a/Engine/Plugins/Experimental/Animation/HierarchyTableAnimation/Source/Runtime/Public/BlendProfileStandalone.h
+++ b/Engine/Plugins/Experimental/Animation/HierarchyTableAnimation/Source/Runtime/Public/BlendProfileStandalone.h
@@ -4,6 +4,7 @@
 
 #include "CoreMinimal.h"
 #include "HierarchyTable.h"
+#include "HierarchyTableBlendProfile.h"
 
 #include "BlendProfileStandalone.generated.h"
 
@@ -15,20 +16,51 @@ enum class EBlendProfileStandaloneType
 	BlendMask
 };
 
-UCLASS()
+UCLASS(EditInlineNew, BlueprintType)
 class HIERARCHYTABLEANIMATIONRUNTIME_API UBlendProfileStandalone : public UObject
 {
 	GENERATED_BODY()
 
+public:
+	// Begin UObject
+	virtual void PostLoad() override;
+#if WITH_EDITOR
+	virtual void PreSave(FObjectPreSaveContext ObjectSaveContext) override;
+	virtual void GetAssetRegistryTags(FAssetRegistryTagsContext Context) const override;
+#endif
+	// End UObject
+
+#if WITH_EDITOR
+	void UpdateHierarchy();
+
+	void UpdateCachedData();
+#endif
+
+	TObjectPtr<USkeleton> GetSkeleton() const;
+
 public:
 	UPROPERTY(VisibleAnywhere, Category=Default, AssetRegistrySearchable)
 	EBlendProfileStandaloneType Type;
 
-	UPROPERTY(VisibleAnywhere, Category=Default)
+#if WITH_EDITORONLY_DATA
+	UPROPERTY()
 	TObjectPtr<UHierarchyTable> Table;
+#endif // WITH_EDITORONLY_DATA
+
+	// The flattened version of the table data for use at runtime instead of,
+	// slowly traversing the table tree.
+	UPROPERTY()
+	FBlendProfileStandaloneCachedData CachedBlendProfileData;
+
+	UPROPERTY(VisibleAnywhere, Category=Default)
+	TObjectPtr<USkeleton> Skeleton;
+
+private:
+	UPROPERTY()
+	FGuid SkeletonHierarchyGuid;
+
+	UPROPERTY()
+	FGuid SkeletonVirtualBonesHierarchyGuid;
 
-#if WITH_EDITOR
-	virtual void GetAssetRegistryTags(FAssetRegistryTagsContext Context) const override;
-#endif
 };
 
diff --git a/Engine/Plugins/Experimental/Animation/HierarchyTableAnimation/Source/Runtime/Public/HierarchyTableBlendProfile.h b/Engine/Plugins/Experimental/Animation/HierarchyTableAnimation/Source/Runtime/Public/HierarchyTableBlendProfile.h
--- a/Engine/Plugins/Experimental/Animation/HierarchyTableAnimation/Source/Runtime/Public/HierarchyTableBlendProfile.h
+++ b/Engine/Plugins/Experimental/Animation/HierarchyTableAnimation/Source/Runtime/Public/HierarchyTableBlendProfile.h
@@ -6,6 +6,108 @@
 #include "Animation/BlendProfile.h"
 #include "HierarchyTable.h"
 
+#include "HierarchyTableBlendProfile.generated.h"
+
+USTRUCT()
+struct FMaskedAttributeWeightSerialised
+{
+	GENERATED_BODY()
+
+	UPROPERTY()
+	FName AttributeNamespace;
+	
+	UPROPERTY()
+	FName AttributeName;
+	
+	UPROPERTY()
+	int32 AttributeIndex;
+
+	UPROPERTY()
+	float Weight;
+};
+
+USTRUCT()
+struct FMaskedCurveWeightSerialised
+{
+	GENERATED_BODY()
+
+	UPROPERTY()
+	FName CurveName;
+
+	UPROPERTY()
+	float Weight;
+};
+
+USTRUCT()
+struct HIERARCHYTABLEANIMATIONRUNTIME_API FBlendProfileStandaloneCachedData
+{
+	GENERATED_BODY()
+
+public:
+	void UnpackCachedData();
+
+	void Init(TObjectPtr<UHierarchyTable> InHierarchyTable, const EBlendProfileMode InMode);
+
+	void Reset();
+
+	struct FNamedFloat
+	{
+		FName Name;
+		float Value;
+
+		FNamedFloat(const FName InName, const float InValue)
+			: Name(InName)
+			, Value(InValue)
+		{
+		}
+	};
+
+	struct FMaskedAttributeWeight
+	{
+		UE::Anim::FAttributeId Attribute;
+		float Weight;
+
+		FMaskedAttributeWeight(const UE::Anim::FAttributeId InAttribute, const float InWeight)
+			: Attribute(InAttribute)
+			, Weight(InWeight)
+		{
+		}
+	};
+
+	const TArray<float>& GetBoneBlendWeights() const;
+
+	const UE::Anim::TNamedValueArray<FDefaultAllocator, UE::Anim::FCurveElement>& GetCurveBlendWeights() const;
+
+	const TArray<FMaskedAttributeWeight>& GetAttributeBlendWeights() const;
+
+	void ConstructBlendProfile(const TObjectPtr<UBlendProfile> OutBlendProfile) const;
+
+	bool IsValidBoneIndex(const int32 BoneIndex) const;
+	
+	TObjectPtr<USkeleton> GetSkeleton() const;
+
+private:
+	UPROPERTY()
+	TObjectPtr<USkeleton> Skeleton;
+
+	UPROPERTY()
+	TArray<float> BoneBlendWeights;
+
+	UPROPERTY()
+	TArray<FMaskedCurveWeightSerialised> SerializedCurveBlendWeights;
+
+	UPROPERTY()
+	TArray<FMaskedAttributeWeightSerialised> SerializedAttributeBlendWeights;
+
+	UPROPERTY()
+	EBlendProfileMode Mode = EBlendProfileMode::WeightFactor;
+
+private:
+	UE::Anim::TNamedValueArray<FDefaultAllocator, UE::Anim::FCurveElement> CurveBlendWeights;
+
+	TArray<FMaskedAttributeWeight> AttributeBlendWeights;
+};
+
 class HIERARCHYTABLEANIMATIONRUNTIME_API FHierarchyTableBlendProfile : public IBlendProfileInterface
 {
 public:
diff --git a/Engine/Plugins/Experimental/Animation/HierarchyTableAnimation/Source/UncookedOnly/Private/AnimGraphNode_BlendProfileLayeredBlend.cpp b/Engine/Plugins/Experimental/Animation/HierarchyTableAnimation/Source/UncookedOnly/Private/AnimGraphNode_BlendProfileLayeredBlend.cpp
--- a/Engine/Plugins/Experimental/Animation/HierarchyTableAnimation/Source/UncookedOnly/Private/AnimGraphNode_BlendProfileLayeredBlend.cpp
+++ b/Engine/Plugins/Experimental/Animation/HierarchyTableAnimation/Source/UncookedOnly/Private/AnimGraphNode_BlendProfileLayeredBlend.cpp
@@ -12,6 +12,7 @@
 
 #include "HierarchyTable.h"
 #include "MaskProfile/HierarchyTableTypeMask.h"
+#include "BlendProfileStandalone.h"
 
 #define LOCTEXT_NAMESPACE "AnimGraphNode_BlendProfileLayeredBlend"
 
@@ -42,10 +43,9 @@ FString UAnimGraphNode_BlendProfileLayeredBlend::GetNodeCategory() const
 
 void UAnimGraphNode_BlendProfileLayeredBlend::PreloadRequiredAssets()
 {
-	TObjectPtr<UHierarchyTable> MaskTable = Node.MaskTable;
-	if (MaskTable)
+	if (Node.BlendProfileAsset)
 	{
-		PreloadObject(MaskTable);
+		PreloadObject(Node.BlendProfileAsset);
 	}
 
 	Super::PreloadRequiredAssets();
@@ -57,12 +57,21 @@ void UAnimGraphNode_BlendProfileLayeredBlend::ValidateAnimNodeDuringCompilation(
 
 	bool bCompilationError = false;
 
-	const TObjectPtr<UHierarchyTable> MaskTable = Node.MaskTable;
+	const TObjectPtr<UBlendProfileStandalone> BlendProfileAsset = Node.BlendProfileAsset;
 
-	if (MaskTable && !MaskTable->IsElementType<FHierarchyTable_ElementType_Mask>())
+	if (BlendProfileAsset)
 	{
-		MessageLog.Error(*LOCTEXT("InvalidMaskType", "@@ has a mask table that is not a mask type.").ToString(), this, MaskTable);
-		bCompilationError = true;
+		if (BlendProfileAsset->Type != EBlendProfileStandaloneType::BlendMask)
+		{
+			MessageLog.Error(*LOCTEXT("InvalidBlendProfileAssetType", "@@ uses a blend profile asset of the incorrect type, expected Blend Mask type.").ToString(), this);
+			bCompilationError = true;
+		}
+
+		if (BlendProfileAsset->GetSkeleton() != ForSkeleton)
+		{
+			MessageLog.Error(*LOCTEXT("InvalidBlendProfileSkeleton", "@@ uses a blend profile asset with the wrong skeleton.").ToString(), this);
+			bCompilationError = true;
+		}
 	}
 }
 
