Profile blend fix - Remap bone weights to compact bone indices in UpdateCachedBoneData

diff --git a/Engine/Plugins/Experimental/Animation/HierarchyTableAnimation/Source/Runtime/Private/AnimNode_BlendProfileLayeredBlend.cpp b/Engine/Plugins/Experimental/Animation/HierarchyTableAnimation/Source/Runtime/Private/AnimNode_BlendProfileLayeredBlend.cpp
--- a/Engine/Plugins/Experimental/Animation/HierarchyTableAnimation/Source/Runtime/Private/AnimNode_BlendProfileLayeredBlend.cpp
+++ b/Engine/Plugins/Experimental/Animation/HierarchyTableAnimation/Source/Runtime/Private/AnimNode_BlendProfileLayeredBlend.cpp
@@ -48,7 +48,13 @@ void FAnimNode_BlendProfileLayeredBlend::UpdateCachedBoneData(const FBoneContain
 	{
 		if (ensure(NumRequiredBones <= BlendProfileAsset->CachedBlendProfileData.GetBoneBlendWeights().Num()))
 		{
-			DesiredBoneBlendWeights = BlendProfileAsset->CachedBlendProfileData.GetBoneBlendWeights();
+			DesiredBoneBlendWeights.SetNumUninitialized(NumRequiredBones);
+			const TArray<float>& BoneBlendWeights = BlendProfileAsset->CachedBlendProfileData.GetBoneBlendWeights();
+			for (int32 RequiredBoneIndex = 0; RequiredBoneIndex < NumRequiredBones; ++RequiredBoneIndex)
+			{
+				const int32 SkeletonBoneIndex = RequiredBones.GetSkeletonIndex(FCompactPoseBoneIndex(RequiredBoneIndex));
+				DesiredBoneBlendWeights[RequiredBoneIndex] = BoneBlendWeights[SkeletonBoneIndex];
+			}
 		}
 	}
 
