HierarchyTableAnimation: Call super

diff --git a/Engine/Plugins/Experimental/Animation/HierarchyTableAnimation/Source/Runtime/Private/BlendProfileStandalone.cpp b/Engine/Plugins/Experimental/Animation/HierarchyTableAnimation/Source/Runtime/Private/BlendProfileStandalone.cpp
--- a/Engine/Plugins/Experimental/Animation/HierarchyTableAnimation/Source/Runtime/Private/BlendProfileStandalone.cpp
+++ b/Engine/Plugins/Experimental/Animation/HierarchyTableAnimation/Source/Runtime/Private/BlendProfileStandalone.cpp
@@ -195,6 +195,8 @@ void UBlendProfileStandalone::PostLoad()
 void UBlendProfileStandalone::PreSave(FObjectPreSaveContext ObjectSaveContext)
 {
 	UpdateCachedData();
+	
+	Super::PreSave(ObjectSaveContext);
 }
 
 void UBlendProfileStandalone::GetAssetRegistryTags(FAssetRegistryTagsContext Context) const
