Add table null check when updating hierarchy in blend profile asset

diff --git a/Engine/Plugins/Experimental/Animation/HierarchyTableAnimation/Source/Runtime/Private/BlendProfileStandalone.cpp b/Engine/Plugins/Experimental/Animation/HierarchyTableAnimation/Source/Runtime/Private/BlendProfileStandalone.cpp
--- a/Engine/Plugins/Experimental/Animation/HierarchyTableAnimation/Source/Runtime/Private/BlendProfileStandalone.cpp
+++ b/Engine/Plugins/Experimental/Animation/HierarchyTableAnimation/Source/Runtime/Private/BlendProfileStandalone.cpp
@@ -26,6 +26,11 @@ void UBlendProfileStandalone::UpdateHierarchy()
 	 *		5. Reconstruct the curves and attributes hierarchy.
 	 */
 
+	if (!Table)
+	{
+		return;
+	}
+
 	const FHierarchyTable_TableType_Skeleton& Metadata = Table->GetTableMetadata<FHierarchyTable_TableType_Skeleton>();
 	Skeleton = Metadata.Skeleton;
 
