Add live updating blend mask weights in the Profile Blend node
* Making edits to a blend profile asset are now immediately reflected in the Profile Blend node without needing to recompile.
* This is editor-only.
* Also add undo for edits to blend profile weights.

diff --git a/Engine/Plugins/Experimental/Animation/HierarchyTable/Source/Editor/Private/SHierarchyTableRow.cpp b/Engine/Plugins/Experimental/Animation/HierarchyTable/Source/Editor/Private/SHierarchyTableRow.cpp
--- a/Engine/Plugins/Experimental/Animation/HierarchyTable/Source/Editor/Private/SHierarchyTableRow.cpp
+++ b/Engine/Plugins/Experimental/Animation/HierarchyTable/Source/Editor/Private/SHierarchyTableRow.cpp
@@ -127,7 +127,11 @@ TSharedRef<SWidget> SHierarchyTableRow::GenerateWidgetForColumn(const FName& Col
 				{
 					const FScopedTransaction Transaction(LOCTEXT("ToggleOverride", "Toggle Override"));
 					HierarchyTableWidget->HierarchyTable->Modify();
+					
+					// TODO: Update API to avoid having to manually regenerate guid
 					HierarchyTableWidget->HierarchyTable->GetMutableTableEntry(TreeItem->Index)->ToggleOverridden();
+					HierarchyTableWidget->HierarchyTable->RegenerateEntriesGuid();
+					
 					return FReply::Handled();
 				})
 			.ContentPadding(2.0f)
diff --git a/Engine/Plugins/Experimental/Animation/HierarchyTable/Source/Runtime/Private/HierarchyTable.cpp b/Engine/Plugins/Experimental/Animation/HierarchyTable/Source/Runtime/Private/HierarchyTable.cpp
--- a/Engine/Plugins/Experimental/Animation/HierarchyTable/Source/Runtime/Private/HierarchyTable.cpp
+++ b/Engine/Plugins/Experimental/Animation/HierarchyTable/Source/Runtime/Private/HierarchyTable.cpp
@@ -222,3 +222,15 @@ int32 UHierarchyTable::GetTableEntryIndex(const FName EntryIdentifier) const
 
 	return INDEX_NONE;
 }
+
+#if WITH_EDITOR
+FGuid UHierarchyTable::GetEntriesGuid()
+{
+	return EntriesGuid;
+}
+
+void UHierarchyTable::RegenerateEntriesGuid()
+{
+	EntriesGuid = FGuid::NewGuid(); 
+}
+#endif //WITH_EDITOR
\ No newline at end of file
diff --git a/Engine/Plugins/Experimental/Animation/HierarchyTable/Source/Runtime/Public/HierarchyTable.h b/Engine/Plugins/Experimental/Animation/HierarchyTable/Source/Runtime/Public/HierarchyTable.h
--- a/Engine/Plugins/Experimental/Animation/HierarchyTable/Source/Runtime/Public/HierarchyTable.h
+++ b/Engine/Plugins/Experimental/Animation/HierarchyTable/Source/Runtime/Public/HierarchyTable.h
@@ -195,11 +195,25 @@ public:
 
 	UE_API FGuid GetHierarchyGuid();
 
+#if WITH_EDITOR
+	// Returns a GUID encoding the state of the values of all entries
+	// i.e. this GUID is invalidated when values inside the table are modified
+	UE_API FGuid GetEntriesGuid();
+	
+	// TODO: Make private in subsequent API changes
+	UE_API void RegenerateEntriesGuid();
+#endif
+
 private:
 	UE_API void RegenerateHierarchyGuid();
 
 	UPROPERTY()
 	FGuid HierarchyGuid;
+	
+#if WITH_EDITORONLY_DATA
+	UPROPERTY()
+    FGuid EntriesGuid;
+#endif
 };
 
 #undef UE_API
diff --git a/Engine/Plugins/Experimental/Animation/HierarchyTableAnimation/Source/Editor/Private/BlendProfileStandaloneEditorToolkit.cpp b/Engine/Plugins/Experimental/Animation/HierarchyTableAnimation/Source/Editor/Private/BlendProfileStandaloneEditorToolkit.cpp
--- a/Engine/Plugins/Experimental/Animation/HierarchyTableAnimation/Source/Editor/Private/BlendProfileStandaloneEditorToolkit.cpp
+++ b/Engine/Plugins/Experimental/Animation/HierarchyTableAnimation/Source/Editor/Private/BlendProfileStandaloneEditorToolkit.cpp
@@ -12,6 +12,8 @@
 void FBlendProfileStandaloneEditorToolkit::InitEditor(const TArray<UObject*>& InObjects)
 {
 	BlendProfile = CastChecked<UBlendProfileStandalone>(InObjects[0]);
+	
+	BlendProfile->Table->SetFlags(RF_Transactional);
 	BlendProfile->UpdateHierarchy();
 
 	const TSharedRef<FTabManager::FLayout> Layout = FTabManager::NewLayout("BlendProfileStandaloneEditorToolkit")
diff --git a/Engine/Plugins/Experimental/Animation/HierarchyTableAnimation/Source/Editor/Private/MaskProfile/MaskProfileColumn.cpp b/Engine/Plugins/Experimental/Animation/HierarchyTableAnimation/Source/Editor/Private/MaskProfile/MaskProfileColumn.cpp
--- a/Engine/Plugins/Experimental/Animation/HierarchyTableAnimation/Source/Editor/Private/MaskProfile/MaskProfileColumn.cpp
+++ b/Engine/Plugins/Experimental/Animation/HierarchyTableAnimation/Source/Editor/Private/MaskProfile/MaskProfileColumn.cpp
@@ -5,6 +5,7 @@
 #include "MaskProfile/HierarchyTableTypeMask.h"
 #include "HierarchyTable.h"
 #include "Editor.h"
+#include "ScopedTransaction.h"
 
 #define LOCTEXT_NAMESPACE "FHierarchyTableColumn_Mask"
 
@@ -19,9 +20,14 @@ TSharedRef<SWidget> FHierarchyTableColumn_Mask::CreateEntryWidget(TObjectPtr<UHi
 			{
 				return HierarchyTable->GetTableEntry(EntryIndex)->GetValue<FHierarchyTable_ElementType_Mask>()->Value;
 			})
-		.OnValueChanged_Lambda([HierarchyTable, EntryIndex](float NewValue)
+		.OnValueChanged_Lambda([HierarchyTable, EntryIndex](const float NewValue)
 			{
+				const FScopedTransaction Transaction(LOCTEXT("ModifyValue", "Modify Value"));
+				HierarchyTable->Modify();
+			
+				// TODO: Update API to avoid having to manually regenerate guid
 				HierarchyTable->GetMutableTableEntry(EntryIndex)->GetMutableValue<FHierarchyTable_ElementType_Mask>()->Value = NewValue;
+				HierarchyTable->RegenerateEntriesGuid();
 			})
 		.OnBeginSliderMovement_Lambda([HierarchyTable]()
 			{
diff --git a/Engine/Plugins/Experimental/Animation/HierarchyTableAnimation/Source/Runtime/Private/AnimNode_BlendProfileLayeredBlend.cpp b/Engine/Plugins/Experimental/Animation/HierarchyTableAnimation/Source/Runtime/Private/AnimNode_BlendProfileLayeredBlend.cpp
--- a/Engine/Plugins/Experimental/Animation/HierarchyTableAnimation/Source/Runtime/Private/AnimNode_BlendProfileLayeredBlend.cpp
+++ b/Engine/Plugins/Experimental/Animation/HierarchyTableAnimation/Source/Runtime/Private/AnimNode_BlendProfileLayeredBlend.cpp
@@ -25,15 +25,30 @@ void FAnimNode_BlendProfileLayeredBlend::Initialize_AnyThread(const FAnimationIn
 
 bool FAnimNode_BlendProfileLayeredBlend::ArePerBoneBlendWeightsValid(const USkeleton* InSkeleton) const
 {
-	return (InSkeleton != nullptr && InSkeleton->GetGuid() == SkeletonGuid && InSkeleton->GetVirtualBoneGuid() == VirtualBoneGuid && CachedBlendProfile == BlendProfileAsset);
+	return InSkeleton != nullptr
+		&& InSkeleton->GetGuid() == SkeletonGuid
+		&& InSkeleton->GetVirtualBoneGuid() == VirtualBoneGuid
+		&& CachedBlendProfile == BlendProfileAsset
+#if WITH_EDITOR
+		/* Support editor-only live updating blend weights when the blend profile values change */
+		&& (BlendProfileAsset && !BlendProfileAsset->IsCachedDataStale());
+#else
+		;
+#endif
 }
-
 void FAnimNode_BlendProfileLayeredBlend::UpdateCachedBoneData(const FBoneContainer& RequiredBones, USkeleton* Skeleton)
 {
 	if (ArePerBoneBlendWeightsValid(Skeleton) && RequiredBonesSerialNumber == RequiredBones.GetSerialNumber())
 	{
 		return;
 	}
+	
+#if WITH_EDITOR
+	if (BlendProfileAsset && BlendProfileAsset->IsCachedDataStale())
+	{
+		BlendProfileAsset->UpdateCachedData();
+	}
+#endif // WITH_EDITOR
 
 	QUICK_SCOPE_CYCLE_COUNTER(STAT_BlendProfileLayeredBlend_UpdateCachedBoneData);
 	
diff --git a/Engine/Plugins/Experimental/Animation/HierarchyTableAnimation/Source/Runtime/Private/BlendProfileStandalone.cpp b/Engine/Plugins/Experimental/Animation/HierarchyTableAnimation/Source/Runtime/Private/BlendProfileStandalone.cpp
--- a/Engine/Plugins/Experimental/Animation/HierarchyTableAnimation/Source/Runtime/Private/BlendProfileStandalone.cpp
+++ b/Engine/Plugins/Experimental/Animation/HierarchyTableAnimation/Source/Runtime/Private/BlendProfileStandalone.cpp
@@ -155,6 +155,13 @@ void UBlendProfileStandalone::UpdateHierarchy()
 
 void UBlendProfileStandalone::UpdateCachedData()
 {
+#if WITH_EDITORONLY_DATA
+	if (!IsCachedDataStale())
+	{
+		return;
+	}
+#endif
+
 	CachedBlendProfileData.Reset();
 
 	if (Table)
@@ -174,6 +181,15 @@ void UBlendProfileStandalone::UpdateCachedData()
 		const FHierarchyTable_TableType_Skeleton& Metadata = Table->GetTableMetadata<FHierarchyTable_TableType_Skeleton>();
 		Skeleton = Metadata.Skeleton;
 	}
+	
+#if WITH_EDITORONLY_DATA
+	CachedEntriesGuid = Table ? Table->GetEntriesGuid() : FGuid();
+#endif
+}
+
+bool UBlendProfileStandalone::IsCachedDataStale() const
+{
+	return Table ? Table->GetEntriesGuid() != CachedEntriesGuid : false;
 }
 
 #endif // WITH_EDITOR
diff --git a/Engine/Plugins/Experimental/Animation/HierarchyTableAnimation/Source/Runtime/Public/AnimNode_BlendProfileLayeredBlend.h b/Engine/Plugins/Experimental/Animation/HierarchyTableAnimation/Source/Runtime/Public/AnimNode_BlendProfileLayeredBlend.h
--- a/Engine/Plugins/Experimental/Animation/HierarchyTableAnimation/Source/Runtime/Public/AnimNode_BlendProfileLayeredBlend.h
+++ b/Engine/Plugins/Experimental/Animation/HierarchyTableAnimation/Source/Runtime/Public/AnimNode_BlendProfileLayeredBlend.h
@@ -99,6 +99,7 @@ private:
 	TArray<float> DesiredBoneBlendWeights;
 	TArray<float> CurrentBoneBlendWeights;
 
+	UPROPERTY()
 	TObjectPtr<UBlendProfileStandalone> CachedBlendProfile;
 };
 
diff --git a/Engine/Plugins/Experimental/Animation/HierarchyTableAnimation/Source/Runtime/Public/BlendProfileStandalone.h b/Engine/Plugins/Experimental/Animation/HierarchyTableAnimation/Source/Runtime/Public/BlendProfileStandalone.h
--- a/Engine/Plugins/Experimental/Animation/HierarchyTableAnimation/Source/Runtime/Public/BlendProfileStandalone.h
+++ b/Engine/Plugins/Experimental/Animation/HierarchyTableAnimation/Source/Runtime/Public/BlendProfileStandalone.h
@@ -36,6 +36,8 @@ public:
 	UE_API void UpdateHierarchy();
 
 	UE_API void UpdateCachedData();
+	
+	UE_API bool IsCachedDataStale() const;
 #endif
 
 	UE_API TObjectPtr<USkeleton> GetSkeleton() const;
@@ -63,7 +65,11 @@ private:
 
 	UPROPERTY()
 	FGuid SkeletonVirtualBonesHierarchyGuid;
-
+	
+#if WITH_EDITORONLY_DATA
+	UPROPERTY()
+    FGuid CachedEntriesGuid;
+#endif
 };
 
 #undef UE_API
