Fix profile blend node crash due to cached data not being generated in some cases

#rnx

diff --git a/Engine/Plugins/Experimental/Animation/HierarchyTableAnimation/Source/Runtime/Private/BlendProfileStandalone.cpp b/Engine/Plugins/Experimental/Animation/HierarchyTableAnimation/Source/Runtime/Private/BlendProfileStandalone.cpp
--- a/Engine/Plugins/Experimental/Animation/HierarchyTableAnimation/Source/Runtime/Private/BlendProfileStandalone.cpp
+++ b/Engine/Plugins/Experimental/Animation/HierarchyTableAnimation/Source/Runtime/Private/BlendProfileStandalone.cpp
@@ -151,6 +151,8 @@ void UBlendProfileStandalone::UpdateHierarchy()
 			}
 		}
 	}
+
+	Table->RegenerateEntriesGuid();
 }
 
 void UBlendProfileStandalone::UpdateCachedData()
