Fix loaded blend profile assets not updating the hierarchy when its skeleton's hierarchy has changed
* Previously blend profile assets only updates its hierarchy on load when a change is detected.
* However, if is possible for a skeleton's hierarchy to change while a blend profile asset is loaded leading to OOBs.

diff --git a/Engine/Plugins/Experimental/Animation/HierarchyTableAnimation/Source/Runtime/Private/BlendProfileStandalone.cpp b/Engine/Plugins/Experimental/Animation/HierarchyTableAnimation/Source/Runtime/Private/BlendProfileStandalone.cpp
--- a/Engine/Plugins/Experimental/Animation/HierarchyTableAnimation/Source/Runtime/Private/BlendProfileStandalone.cpp
+++ b/Engine/Plugins/Experimental/Animation/HierarchyTableAnimation/Source/Runtime/Private/BlendProfileStandalone.cpp
@@ -16,6 +16,22 @@
 
 #if WITH_EDITOR
 
+void UBlendProfileStandalone::BeginDestroy()
+{
+	if (Skeleton && !IsTemplate())
+	{
+		Skeleton->UnregisterOnSkeletonHierarchyChanged(this);
+	}
+
+	Super::BeginDestroy();
+}
+
+void UBlendProfileStandalone::HandleSkeletonHierarchyChanged()
+{
+	UpdateHierarchy();
+	UpdateCachedData();
+}
+
 void UBlendProfileStandalone::UpdateHierarchy()
 {
 	/* Skeleton hierarchy tables take a snapshot of the currently bound skeleton's hierarchy.
@@ -34,7 +50,7 @@ void UBlendProfileStandalone::UpdateHierarchy()
 	}
 
 	const FHierarchyTable_TableType_Skeleton& Metadata = Table->GetTableMetadata<FHierarchyTable_TableType_Skeleton>();
-	Skeleton = Metadata.Skeleton;
+	SetSkeleton(Metadata.Skeleton);
 
 	if (Skeleton->GetGuid() == SkeletonHierarchyGuid && Skeleton->GetVirtualBoneGuid() == SkeletonVirtualBonesHierarchyGuid)
 	{
@@ -181,7 +197,7 @@ void UBlendProfileStandalone::UpdateCachedData()
 		CachedBlendProfileData.Init(Table, Mode);
 
 		const FHierarchyTable_TableType_Skeleton& Metadata = Table->GetTableMetadata<FHierarchyTable_TableType_Skeleton>();
-		Skeleton = Metadata.Skeleton;
+		SetSkeleton(Metadata.Skeleton);
 	}
 	
 #if WITH_EDITORONLY_DATA
@@ -236,4 +252,20 @@ void UBlendProfileStandalone::GetAssetRegistryTags(FAssetRegistryTagsContext Con
 	}
 }
 
+void UBlendProfileStandalone::SetSkeleton(TObjectPtr<USkeleton> InSkeleton)
+{
+	if (Skeleton)
+	{
+		Skeleton->UnregisterOnSkeletonHierarchyChanged(this);
+	}
+	
+	if (ensure(InSkeleton))
+	{
+		Skeleton = InSkeleton;
+
+		Skeleton->RegisterOnSkeletonHierarchyChanged(
+			USkeleton::FOnSkeletonHierarchyChanged::CreateUObject(this, &UBlendProfileStandalone::HandleSkeletonHierarchyChanged));
+	}
+}
+
 #endif // WITH_EDITOR
diff --git a/Engine/Plugins/Experimental/Animation/HierarchyTableAnimation/Source/Runtime/Public/BlendProfileStandalone.h b/Engine/Plugins/Experimental/Animation/HierarchyTableAnimation/Source/Runtime/Public/BlendProfileStandalone.h
--- a/Engine/Plugins/Experimental/Animation/HierarchyTableAnimation/Source/Runtime/Public/BlendProfileStandalone.h
+++ b/Engine/Plugins/Experimental/Animation/HierarchyTableAnimation/Source/Runtime/Public/BlendProfileStandalone.h
@@ -27,6 +27,7 @@ public:
 	// Begin UObject
 	UE_API virtual void PostLoad() override;
 #if WITH_EDITOR
+	UE_API virtual void BeginDestroy() override;
 	UE_API virtual void PreSave(FObjectPreSaveContext ObjectSaveContext) override;
 	UE_API virtual void GetAssetRegistryTags(FAssetRegistryTagsContext Context) const override;
 #endif
@@ -56,10 +57,16 @@ public:
 	UPROPERTY()
 	FBlendProfileStandaloneCachedData CachedBlendProfileData;
 
+private:
+#if WITH_EDITOR
+	void SetSkeleton(TObjectPtr<USkeleton> InSkeleton);
+
+	void HandleSkeletonHierarchyChanged();
+#endif // WITH_EDITOR
+
 	UPROPERTY(VisibleAnywhere, Category=Default)
 	TObjectPtr<USkeleton> Skeleton;
 
-private:
 	UPROPERTY()
 	FGuid SkeletonHierarchyGuid;
 
